MediaSourceManager: Added media source manager plugin setup.

#rnx

diff --git a/Engine/Plugins/Media/MediaSourceManager/MediaSourceManager.uplugin b/Engine/Plugins/Media/MediaSourceManager/MediaSourceManager.uplugin
new file mode 100644
--- /dev/null
+++ b/Engine/Plugins/Media/MediaSourceManager/MediaSourceManager.uplugin
@@ -0,0 +1,26 @@
+{
+	"FileVersion": 3,
+	"Version": 0,
+	"VersionName": "0.1",
+	"FriendlyName": "Media Source Manager",
+	"Description": "Manager to handle media sources.",
+	"Category": "Media",
+	"CreatedBy": "Epic Games, Inc.",
+	"CreatedByURL": "http://epicgames.com",
+	"DocsURL": "https://docs.unrealengine.com/en-US/Engine/MediaFramework/Overview",
+	"MarketplaceURL": "",
+	"SupportURL": "",
+	"CanContainContent": true,
+	"EnabledByDefault": false,
+	"IsBetaVersion": false,
+	"IsExperimentalVersion": false,
+	"Installed": false,
+	"Modules": 
+	[
+		{
+			"Name": "MediaSourceManager",
+			"Type": "Editor",
+			"LoadingPhase": "Default"
+		}
+	]
+}
\ No newline at end of file
diff --git a/Engine/Plugins/Media/MediaSourceManager/Source/MediaSourceManager/MediaSourceManager.Build.cs b/Engine/Plugins/Media/MediaSourceManager/Source/MediaSourceManager/MediaSourceManager.Build.cs
new file mode 100644
--- /dev/null
+++ b/Engine/Plugins/Media/MediaSourceManager/Source/MediaSourceManager/MediaSourceManager.Build.cs
@@ -0,0 +1,24 @@
+// Copyright Epic Games, Inc. All Rights Reserved.
+
+using UnrealBuildTool;
+
+public class MediaSourceManager : ModuleRules
+{
+	public MediaSourceManager(ReadOnlyTargetRules Target) : base(Target)
+	{
+		PCHUsage = ModuleRules.PCHUsageMode.UseExplicitOrSharedPCHs;
+		
+		PublicDependencyModuleNames.AddRange(
+			new string[]
+			{
+				"Core",
+			}
+			);
+		
+		PrivateDependencyModuleNames.AddRange(
+			new string[]
+			{
+			}
+			);
+	}
+}
diff --git a/Engine/Plugins/Media/MediaSourceManager/Source/MediaSourceManager/Private/MediaSourceManagerModule.cpp b/Engine/Plugins/Media/MediaSourceManager/Source/MediaSourceManager/Private/MediaSourceManagerModule.cpp
new file mode 100644
--- /dev/null
+++ b/Engine/Plugins/Media/MediaSourceManager/Source/MediaSourceManager/Private/MediaSourceManagerModule.cpp
@@ -0,0 +1,33 @@
+// Copyright Epic Games, Inc. All Rights Reserved.
+
+#include "MediaSourceManagerModule.h"
+#include "Modules/ModuleManager.h"
+
+#define LOCTEXT_NAMESPACE "MediaSourceManagerModule"
+
+DEFINE_LOG_CATEGORY(LogMediaSourceManager);
+
+/**
+ * Implements the MediaSourceManager module.
+ */
+class FMediaSourceManagerModule
+	: public IMediaSourceManagerModule
+{
+public:
+	//~ IMediaSourceManagerModule interface
+	
+
+	//~ IModuleInterface interface
+
+	virtual void StartupModule() override
+	{
+	}
+
+	virtual void ShutdownModule() override
+	{
+	}
+};
+
+IMPLEMENT_MODULE(FMediaSourceManagerModule, MediaSourceManager);
+
+#undef LOCTEXT_NAMESPACE
diff --git a/Engine/Plugins/Media/MediaSourceManager/Source/MediaSourceManager/Public/MediaSourceManagerModule.h b/Engine/Plugins/Media/MediaSourceManager/Source/MediaSourceManager/Public/MediaSourceManagerModule.h
new file mode 100644
--- /dev/null
+++ b/Engine/Plugins/Media/MediaSourceManager/Source/MediaSourceManager/Public/MediaSourceManagerModule.h
@@ -0,0 +1,19 @@
+// Copyright Epic Games, Inc. All Rights Reserved.
+
+#pragma once
+
+#include "Logging/LogMacros.h"
+#include "Modules/ModuleInterface.h"
+
+/** Log category for this module. */
+DECLARE_LOG_CATEGORY_EXTERN(LogMediaSourceManager, Log, All);
+
+/**
+* Interface for the MediaSourceManager module.
+*/
+class MEDIASOURCEMANAGER_API IMediaSourceManagerModule
+	: public IModuleInterface
+{
+public:
+
+};
