MediaSourceManager: Added editor module setup.
#rnx

diff --git a/Engine/Plugins/Media/MediaSourceManager/MediaSourceManager.uplugin b/Engine/Plugins/Media/MediaSourceManager/MediaSourceManager.uplugin
--- a/Engine/Plugins/Media/MediaSourceManager/MediaSourceManager.uplugin
+++ b/Engine/Plugins/Media/MediaSourceManager/MediaSourceManager.uplugin
@@ -19,6 +19,11 @@
 	[
 		{
 			"Name": "MediaSourceManager",
+			"Type": "Runtime",
+			"LoadingPhase": "Default"
+		},
+		{
+			"Name": "MediaSourceManagerEditor",
 			"Type": "Editor",
 			"LoadingPhase": "Default"
 		}
diff --git a/Engine/Plugins/Media/MediaSourceManager/Source/MediaSourceManagerEditor/MediaSourceManagerEditor.Build.cs b/Engine/Plugins/Media/MediaSourceManager/Source/MediaSourceManagerEditor/MediaSourceManagerEditor.Build.cs
new file mode 100644
--- /dev/null
+++ b/Engine/Plugins/Media/MediaSourceManager/Source/MediaSourceManagerEditor/MediaSourceManagerEditor.Build.cs
@@ -0,0 +1,19 @@
+// Copyright Epic Games, Inc. All Rights Reserved.
+
+using UnrealBuildTool;
+
+public class MediaSourceManagerEditor : ModuleRules
+{
+	public MediaSourceManagerEditor(ReadOnlyTargetRules Target) : base(Target)
+	{
+		PCHUsage = ModuleRules.PCHUsageMode.UseExplicitOrSharedPCHs;
+		
+		PrivateDependencyModuleNames.AddRange(
+			new string[]
+			{
+				"Core",
+				"CoreUObject",
+			}
+			);
+	}
+}
diff --git a/Engine/Plugins/Media/MediaSourceManager/Source/MediaSourceManagerEditor/Private/MediaSourceManagerModuleEditor.cpp b/Engine/Plugins/Media/MediaSourceManager/Source/MediaSourceManagerEditor/Private/MediaSourceManagerModuleEditor.cpp
new file mode 100644
--- /dev/null
+++ b/Engine/Plugins/Media/MediaSourceManager/Source/MediaSourceManagerEditor/Private/MediaSourceManagerModuleEditor.cpp
@@ -0,0 +1,33 @@
+// Copyright Epic Games, Inc. All Rights Reserved.
+
+#include "MediaSourceManagerEditorModule.h"
+#include "Modules/ModuleManager.h"
+
+#define LOCTEXT_NAMESPACE "MediaSourceManagerEditorModule"
+
+DEFINE_LOG_CATEGORY(LogMediaSourceManagerEditor);
+
+/**
+ * Implements the MediaSourceManagerEditor module.
+ */
+class FMediaSourceManagerEditorModule
+	: public IMediaSourceManagerEditorModule
+{
+public:
+	//~ IMediaSourceManagerEditorModule interface
+	
+
+	//~ IModuleInterface interface
+
+	virtual void StartupModule() override
+	{
+	}
+
+	virtual void ShutdownModule() override
+	{
+	}
+};
+
+IMPLEMENT_MODULE(FMediaSourceManagerEditorModule, MediaSourceManagerEditor);
+
+#undef LOCTEXT_NAMESPACE
diff --git a/Engine/Plugins/Media/MediaSourceManager/Source/MediaSourceManagerEditor/Public/MediaSourceManagerEditorModule.h b/Engine/Plugins/Media/MediaSourceManager/Source/MediaSourceManagerEditor/Public/MediaSourceManagerEditorModule.h
new file mode 100644
--- /dev/null
+++ b/Engine/Plugins/Media/MediaSourceManager/Source/MediaSourceManagerEditor/Public/MediaSourceManagerEditorModule.h
@@ -0,0 +1,19 @@
+// Copyright Epic Games, Inc. All Rights Reserved.
+
+#pragma once
+
+#include "Logging/LogMacros.h"
+#include "Modules/ModuleInterface.h"
+
+/** Log category for this module. */
+DECLARE_LOG_CATEGORY_EXTERN(LogMediaSourceManagerEditor, Log, All);
+
+/**
+* Interface for the MediaSourceManagerEditor module.
+*/
+class MEDIASOURCEMANAGEREDITOR_API IMediaSourceManagerEditorModule
+	: public IModuleInterface
+{
+public:
+
+};
