MediaSourceManager: Added manager asset support code.

#rnx

diff --git a/Engine/Plugins/Media/MediaSourceManager/Source/MediaSourceManager/MediaSourceManager.Build.cs b/Engine/Plugins/Media/MediaSourceManager/Source/MediaSourceManager/MediaSourceManager.Build.cs
--- a/Engine/Plugins/Media/MediaSourceManager/Source/MediaSourceManager/MediaSourceManager.Build.cs
+++ b/Engine/Plugins/Media/MediaSourceManager/Source/MediaSourceManager/MediaSourceManager.Build.cs
@@ -8,16 +8,11 @@ public class MediaSourceManager : ModuleRules
 	{
 		PCHUsage = ModuleRules.PCHUsageMode.UseExplicitOrSharedPCHs;
 		
-		PublicDependencyModuleNames.AddRange(
-			new string[]
-			{
-				"Core",
-			}
-			);
-		
 		PrivateDependencyModuleNames.AddRange(
 			new string[]
 			{
+				"Core",
+				"CoreUObject",
 			}
 			);
 	}
diff --git a/Engine/Plugins/Media/MediaSourceManager/Source/MediaSourceManager/Private/MediaSourceManager.cpp b/Engine/Plugins/Media/MediaSourceManager/Source/MediaSourceManager/Private/MediaSourceManager.cpp
new file mode 100644
--- /dev/null
+++ b/Engine/Plugins/Media/MediaSourceManager/Source/MediaSourceManager/Private/MediaSourceManager.cpp
@@ -0,0 +1,9 @@
+// Copyright Epic Games, Inc. All Rights Reserved.
+
+#include "MediaSourceManager.h"
+
+#define LOCTEXT_NAMESPACE "MediaSourceManager"
+
+
+
+#undef LOCTEXT_NAMESPACE
diff --git a/Engine/Plugins/Media/MediaSourceManager/Source/MediaSourceManager/Public/MediaSourceManager.h b/Engine/Plugins/Media/MediaSourceManager/Source/MediaSourceManager/Public/MediaSourceManager.h
new file mode 100644
--- /dev/null
+++ b/Engine/Plugins/Media/MediaSourceManager/Source/MediaSourceManager/Public/MediaSourceManager.h
@@ -0,0 +1,21 @@
+// Copyright Epic Games, Inc. All Rights Reserved.
+
+#pragma once
+
+#include "CoreTypes.h"
+#include "UObject/Object.h"
+#include "UObject/ObjectMacros.h"
+
+#include "MediaSourceManager.generated.h"
+
+/**
+* Manager to handle media sources and their connections.
+*/
+UCLASS()
+class MEDIASOURCEMANAGER_API UMediaSourceManager : public UObject
+{
+	GENERATED_BODY()
+
+public:
+
+};
diff --git a/Engine/Plugins/Media/MediaSourceManager/Source/MediaSourceManagerEditor/MediaSourceManagerEditor.Build.cs b/Engine/Plugins/Media/MediaSourceManager/Source/MediaSourceManagerEditor/MediaSourceManagerEditor.Build.cs
--- a/Engine/Plugins/Media/MediaSourceManager/Source/MediaSourceManagerEditor/MediaSourceManagerEditor.Build.cs
+++ b/Engine/Plugins/Media/MediaSourceManager/Source/MediaSourceManagerEditor/MediaSourceManagerEditor.Build.cs
@@ -13,6 +13,8 @@ public class MediaSourceManagerEditor : ModuleRules
 			{
 				"Core",
 				"CoreUObject",
+				"MediaSourceManager",
+				"UnrealEd",
 			}
 			);
 	}
diff --git a/Engine/Plugins/Media/MediaSourceManager/Source/MediaSourceManagerEditor/Private/AssetTools/MediaSourceManagerActions.cpp b/Engine/Plugins/Media/MediaSourceManager/Source/MediaSourceManagerEditor/Private/AssetTools/MediaSourceManagerActions.cpp
new file mode 100644
--- /dev/null
+++ b/Engine/Plugins/Media/MediaSourceManager/Source/MediaSourceManagerEditor/Private/AssetTools/MediaSourceManagerActions.cpp
@@ -0,0 +1,40 @@
+// Copyright Epic Games, Inc. All Rights Reserved.
+
+#include "AssetTools/MediaSourceManagerActions.h"
+
+#include "AssetRegistry/AssetData.h"
+#include "MediaSourceManager.h"
+
+#define LOCTEXT_NAMESPACE "MediaSourceManagerActions"
+
+bool FMediaSourceManagerActions::CanFilter()
+{
+	return true;
+}
+
+FText FMediaSourceManagerActions::GetAssetDescription(const struct FAssetData& AssetData) const
+{
+	return FText::GetEmpty();
+}
+
+uint32 FMediaSourceManagerActions::GetCategories()
+{
+	return EAssetTypeCategories::Media;
+}
+
+FText FMediaSourceManagerActions::GetName() const
+{
+	return LOCTEXT("MediaSourceManager", "Media Source Manager");
+}
+
+UClass* FMediaSourceManagerActions::GetSupportedClass() const
+{
+	return UMediaSourceManager::StaticClass();
+}
+
+FColor FMediaSourceManagerActions::GetTypeColor() const
+{
+	return FColor::White;
+}
+
+#undef LOCTEXT_NAMESPACE
diff --git a/Engine/Plugins/Media/MediaSourceManager/Source/MediaSourceManagerEditor/Private/AssetTools/MediaSourceManagerActions.h b/Engine/Plugins/Media/MediaSourceManager/Source/MediaSourceManagerEditor/Private/AssetTools/MediaSourceManagerActions.h
new file mode 100644
--- /dev/null
+++ b/Engine/Plugins/Media/MediaSourceManager/Source/MediaSourceManagerEditor/Private/AssetTools/MediaSourceManagerActions.h
@@ -0,0 +1,21 @@
+// Copyright Epic Games, Inc. All Rights Reserved.
+
+#pragma once
+
+#include "AssetTypeActions_Base.h"
+
+/**
+ * Implements an action for MediaSourceManager assets.
+ */
+class FMediaSourceManagerActions : public FAssetTypeActions_Base
+{
+public:
+
+	//~ FAssetTypeActions_Base interface
+	virtual bool CanFilter() override;
+	virtual FText GetAssetDescription(const FAssetData& AssetData) const override;
+	virtual uint32 GetCategories() override;
+	virtual FText GetName() const override;
+	virtual UClass* GetSupportedClass() const override;
+	virtual FColor GetTypeColor() const override;
+};
diff --git a/Engine/Plugins/Media/MediaSourceManager/Source/MediaSourceManagerEditor/Private/Factories/MediaSourceManagerFactory.cpp b/Engine/Plugins/Media/MediaSourceManager/Source/MediaSourceManagerEditor/Private/Factories/MediaSourceManagerFactory.cpp
new file mode 100644
--- /dev/null
+++ b/Engine/Plugins/Media/MediaSourceManager/Source/MediaSourceManagerEditor/Private/Factories/MediaSourceManagerFactory.cpp
@@ -0,0 +1,35 @@
+// Copyright Epic Games, Inc. All Rights Reserved.
+
+#include "MediaSourceManagerFactory.h"
+
+#include "AssetTypeCategories.h"
+#include "MediaSourceManager.h"
+
+/* UMediaSourceManagerFactory structors
+ *****************************************************************************/
+
+UMediaSourceManagerFactory::UMediaSourceManagerFactory(const FObjectInitializer& ObjectInitializer)
+	: Super(ObjectInitializer)
+{
+	SupportedClass = UMediaSourceManager::StaticClass();
+	bCreateNew = true;
+	bEditAfterNew = true;
+}
+
+/* UFactory overrides
+ *****************************************************************************/
+
+UObject* UMediaSourceManagerFactory::FactoryCreateNew(UClass* InClass, UObject* InParent, FName InName, EObjectFlags Flags, UObject* Context, FFeedbackContext* Warn)
+{
+	return NewObject<UMediaSourceManager>(InParent, InClass, InName, Flags);
+}
+
+uint32 UMediaSourceManagerFactory::GetMenuCategories() const
+{
+	return EAssetTypeCategories::Media;
+}
+
+bool UMediaSourceManagerFactory::ShouldShowInNewMenu() const
+{
+	return true;
+}
diff --git a/Engine/Plugins/Media/MediaSourceManager/Source/MediaSourceManagerEditor/Private/Factories/MediaSourceManagerFactory.h b/Engine/Plugins/Media/MediaSourceManager/Source/MediaSourceManagerEditor/Private/Factories/MediaSourceManagerFactory.h
new file mode 100644
--- /dev/null
+++ b/Engine/Plugins/Media/MediaSourceManager/Source/MediaSourceManagerEditor/Private/Factories/MediaSourceManagerFactory.h
@@ -0,0 +1,25 @@
+// Copyright Epic Games, Inc. All Rights Reserved.
+
+#pragma once
+
+#include "Factories/Factory.h"
+
+#include "MediaSourceManagerFactory.generated.h"
+
+/**
+ * Implements a factory for UMediaSourceManager objects.
+ */
+UCLASS(hidecategories=Object)
+class UMediaSourceManagerFactory
+	: public UFactory
+{
+	GENERATED_UCLASS_BODY()
+
+public:
+
+	//~ UFactory Interface
+
+	virtual UObject* FactoryCreateNew(UClass* InClass, UObject* InParent, FName InName, EObjectFlags Flags, UObject* Context, FFeedbackContext* Warn) override;
+	virtual uint32 GetMenuCategories() const override;
+	virtual bool ShouldShowInNewMenu() const override;
+};
diff --git a/Engine/Plugins/Media/MediaSourceManager/Source/MediaSourceManagerEditor/Private/MediaSourceManagerModuleEditor.cpp b/Engine/Plugins/Media/MediaSourceManager/Source/MediaSourceManagerEditor/Private/MediaSourceManagerModuleEditor.cpp
--- a/Engine/Plugins/Media/MediaSourceManager/Source/MediaSourceManagerEditor/Private/MediaSourceManagerModuleEditor.cpp
+++ b/Engine/Plugins/Media/MediaSourceManager/Source/MediaSourceManagerEditor/Private/MediaSourceManagerModuleEditor.cpp
@@ -1,6 +1,9 @@
 // Copyright Epic Games, Inc. All Rights Reserved.
 
 #include "MediaSourceManagerEditorModule.h"
+
+#include "AssetToolsModule.h"
+#include "AssetTools/MediaSourceManagerActions.h"
 #include "Modules/ModuleManager.h"
 
 #define LOCTEXT_NAMESPACE "MediaSourceManagerEditorModule"
@@ -21,11 +24,47 @@ public:
 
 	virtual void StartupModule() override
 	{
+		RegisterAssetTools();
 	}
 
 	virtual void ShutdownModule() override
 	{
+		UnregisterAssetTools();
+	}
+
+protected:
+
+	/**
+	 * Register all our asset tools.
+	 */
+	void RegisterAssetTools()
+	{
+		IAssetTools& AssetTools = FModuleManager::LoadModuleChecked<FAssetToolsModule>("AssetTools").Get();
+
+		TSharedRef<IAssetTypeActions> Action = MakeShared<FMediaSourceManagerActions>();
+		AssetTools.RegisterAssetTypeActions(Action);
+		RegisteredAssetTypeActions.Add(Action);
 	}
+
+	/**
+	 * Unregister all our asset tools.
+	 */
+	void UnregisterAssetTools()
+	{
+		FAssetToolsModule* AssetToolsModule = FModuleManager::GetModulePtr<FAssetToolsModule>("AssetTools");
+
+		if (AssetToolsModule != nullptr)
+		{
+			IAssetTools& AssetTools = AssetToolsModule->Get();
+			for (const TSharedRef<IAssetTypeActions>& Action : RegisteredAssetTypeActions)
+			{
+				AssetTools.UnregisterAssetTypeActions(Action);
+			}
+		}
+	}
+
+	/** The collection of registered asset type actions. */
+	TArray<TSharedRef<IAssetTypeActions>> RegisteredAssetTypeActions;
 };
 
 IMPLEMENT_MODULE(FMediaSourceManagerEditorModule, MediaSourceManagerEditor);
