MediaSourceManager: Added editor toolkit.

#rnx

diff --git a/Engine/Plugins/Media/MediaSourceManager/Source/MediaSourceManagerEditor/MediaSourceManagerEditor.Build.cs b/Engine/Plugins/Media/MediaSourceManager/Source/MediaSourceManagerEditor/MediaSourceManagerEditor.Build.cs
--- a/Engine/Plugins/Media/MediaSourceManager/Source/MediaSourceManagerEditor/MediaSourceManagerEditor.Build.cs
+++ b/Engine/Plugins/Media/MediaSourceManager/Source/MediaSourceManagerEditor/MediaSourceManagerEditor.Build.cs
@@ -14,6 +14,8 @@ public class MediaSourceManagerEditor : ModuleRules
 				"Core",
 				"CoreUObject",
 				"MediaSourceManager",
+				"Slate",
+				"SlateCore",
 				"UnrealEd",
 			}
 			);
diff --git a/Engine/Plugins/Media/MediaSourceManager/Source/MediaSourceManagerEditor/Private/AssetTools/MediaSourceManagerActions.cpp b/Engine/Plugins/Media/MediaSourceManager/Source/MediaSourceManagerEditor/Private/AssetTools/MediaSourceManagerActions.cpp
--- a/Engine/Plugins/Media/MediaSourceManager/Source/MediaSourceManagerEditor/Private/AssetTools/MediaSourceManagerActions.cpp
+++ b/Engine/Plugins/Media/MediaSourceManager/Source/MediaSourceManagerEditor/Private/AssetTools/MediaSourceManagerActions.cpp
@@ -4,6 +4,7 @@
 
 #include "AssetRegistry/AssetData.h"
 #include "MediaSourceManager.h"
+#include "Toolkits/MediaSourceManagerEditorToolkit.h"
 
 #define LOCTEXT_NAMESPACE "MediaSourceManagerActions"
 
@@ -37,4 +38,22 @@ FColor FMediaSourceManagerActions::GetTypeColor() const
 	return FColor::White;
 }
 
+void FMediaSourceManagerActions::OpenAssetEditor(const TArray<UObject*>& InObjects, TSharedPtr<class IToolkitHost> EditWithinLevelEditor)
+{
+	EToolkitMode::Type Mode = EditWithinLevelEditor.IsValid()
+		? EToolkitMode::WorldCentric
+		: EToolkitMode::Standalone;
+
+	for (auto ObjIt = InObjects.CreateConstIterator(); ObjIt; ++ObjIt)
+	{
+		UMediaSourceManager* MediaSourceManager = Cast<UMediaSourceManager>(*ObjIt);
+
+		if (MediaSourceManager != nullptr)
+		{
+			TSharedRef<FMediaSourceManagerEditorToolkit> EditorToolkit = MakeShareable(new FMediaSourceManagerEditorToolkit());
+			EditorToolkit->Initialize(MediaSourceManager, Mode, EditWithinLevelEditor);
+		}
+	}
+}
+
 #undef LOCTEXT_NAMESPACE
diff --git a/Engine/Plugins/Media/MediaSourceManager/Source/MediaSourceManagerEditor/Private/AssetTools/MediaSourceManagerActions.h b/Engine/Plugins/Media/MediaSourceManager/Source/MediaSourceManagerEditor/Private/AssetTools/MediaSourceManagerActions.h
--- a/Engine/Plugins/Media/MediaSourceManager/Source/MediaSourceManagerEditor/Private/AssetTools/MediaSourceManagerActions.h
+++ b/Engine/Plugins/Media/MediaSourceManager/Source/MediaSourceManagerEditor/Private/AssetTools/MediaSourceManagerActions.h
@@ -3,6 +3,7 @@
 #pragma once
 
 #include "AssetTypeActions_Base.h"
+#include "Templates/SharedPointer.h"
 
 /**
  * Implements an action for MediaSourceManager assets.
@@ -18,4 +19,5 @@ public:
 	virtual FText GetName() const override;
 	virtual UClass* GetSupportedClass() const override;
 	virtual FColor GetTypeColor() const override;
+	virtual void OpenAssetEditor(const TArray<UObject*>& InObjects, TSharedPtr<class IToolkitHost> EditWithinLevelEditor = TSharedPtr<IToolkitHost>()) override;
 };
diff --git a/Engine/Plugins/Media/MediaSourceManager/Source/MediaSourceManagerEditor/Private/MediaSourceManagerModuleEditor.cpp b/Engine/Plugins/Media/MediaSourceManager/Source/MediaSourceManagerEditor/Private/MediaSourceManagerModuleEditor.cpp
--- a/Engine/Plugins/Media/MediaSourceManager/Source/MediaSourceManagerEditor/Private/MediaSourceManagerModuleEditor.cpp
+++ b/Engine/Plugins/Media/MediaSourceManager/Source/MediaSourceManagerEditor/Private/MediaSourceManagerModuleEditor.cpp
@@ -41,7 +41,17 @@ protected:
 	{
 		IAssetTools& AssetTools = FModuleManager::LoadModuleChecked<FAssetToolsModule>("AssetTools").Get();
 
-		TSharedRef<IAssetTypeActions> Action = MakeShared<FMediaSourceManagerActions>();
+		RegisterAssetTypeAction(AssetTools, MakeShareable(new FMediaSourceManagerActions()));
+	}
+
+	/**
+	 * Registers a single asset type action.
+	 *
+	 * @param AssetTools	The asset tools object to register with.
+	 * @param Action		The asset type action to register.
+	 */
+	void RegisterAssetTypeAction(IAssetTools& AssetTools, TSharedRef<IAssetTypeActions> Action)
+	{
 		AssetTools.RegisterAssetTypeActions(Action);
 		RegisteredAssetTypeActions.Add(Action);
 	}
diff --git a/Engine/Plugins/Media/MediaSourceManager/Source/MediaSourceManagerEditor/Private/Toolkits/MediaSourceManagerEditorToolkit.cpp b/Engine/Plugins/Media/MediaSourceManager/Source/MediaSourceManagerEditor/Private/Toolkits/MediaSourceManagerEditorToolkit.cpp
new file mode 100644
--- /dev/null
+++ b/Engine/Plugins/Media/MediaSourceManager/Source/MediaSourceManagerEditor/Private/Toolkits/MediaSourceManagerEditorToolkit.cpp
@@ -0,0 +1,155 @@
+// Copyright Epic Games, Inc. All Rights Reserved.
+
+#include "Toolkits/MediaSourceManagerEditorToolkit.h"
+
+#include "Editor.h"
+#include "EditorReimportHandler.h"
+#include "MediaSourceManager.h"
+#include "SlateOptMacros.h"
+#include "Styling/AppStyle.h"
+#include "Widgets/Docking/SDockTab.h"
+#include "Widgets/SMediaSourceManagerSources.h"
+
+#define LOCTEXT_NAMESPACE "FMediaSourceManagerEditorToolkit"
+
+namespace MediaSourceManagerEditorToolkit
+{
+	static const FName AppIdentifier("MediaSourceManagerEditorApp");
+	static const FName SourcesTabId("Sources");
+}
+
+FMediaSourceManagerEditorToolkit::FMediaSourceManagerEditorToolkit()
+	: MediaSourceManager(nullptr)
+{
+}
+
+FMediaSourceManagerEditorToolkit::~FMediaSourceManagerEditorToolkit()
+{
+	FReimportManager::Instance()->OnPreReimport().RemoveAll(this);
+	FReimportManager::Instance()->OnPostReimport().RemoveAll(this);
+
+	GEditor->UnregisterForUndo(this);
+}
+
+void FMediaSourceManagerEditorToolkit::Initialize(UMediaSourceManager* InMediaSourceManager,
+	const EToolkitMode::Type InMode, const TSharedPtr<IToolkitHost>& InToolkitHost)
+{
+	MediaSourceManager = InMediaSourceManager;
+
+	if (MediaSourceManager == nullptr)
+	{
+		return;
+	}
+
+	// support undo/redo
+	MediaSourceManager->SetFlags(RF_Transactional);
+	GEditor->RegisterForUndo(this);
+
+	// create tab layout
+	const TSharedRef<FTabManager::FLayout> Layout = FTabManager::NewLayout("Standalone_MediaSourceManagerEditor_v0.1b")
+		->AddArea
+		(
+			FTabManager::NewPrimaryArea()
+				->SetOrientation(Orient_Vertical)
+				->Split
+				(
+					// Sources tab.
+					FTabManager::NewStack()
+						->AddTab(MediaSourceManagerEditorToolkit::SourcesTabId, ETabState::OpenedTab)
+						->SetHideTabWell(true)
+				)
+		);
+
+	FAssetEditorToolkit::InitAssetEditor(
+		InMode,
+		InToolkitHost,
+		MediaSourceManagerEditorToolkit::AppIdentifier,
+		Layout,
+		true,
+		true,
+		MediaSourceManager
+	);
+	
+	RegenerateMenusAndToolbars();
+}
+
+FString FMediaSourceManagerEditorToolkit::GetDocumentationLink() const
+{
+	return FString(TEXT("WorkingWithMedia/IntegratingMedia/MediaFramework"));
+}
+
+void FMediaSourceManagerEditorToolkit::RegisterTabSpawners(const TSharedRef<FTabManager>& InTabManager)
+{
+	WorkspaceMenuCategory = InTabManager->AddLocalWorkspaceMenuCategory(
+		LOCTEXT("WorkspaceMenu_MediaSourcesManager", "Media Sources Manager Editor"));
+	auto WorkspaceMenuCategoryRef = WorkspaceMenuCategory.ToSharedRef();
+
+	FAssetEditorToolkit::RegisterTabSpawners(InTabManager);
+
+	// Sources tab.
+	InTabManager->RegisterTabSpawner(MediaSourceManagerEditorToolkit::SourcesTabId,
+		FOnSpawnTab::CreateSP(this, &FMediaSourceManagerEditorToolkit::HandleTabManagerSpawnTab, 
+			MediaSourceManagerEditorToolkit::SourcesTabId))
+		.SetDisplayName(LOCTEXT("SourcesTabName", "Sources"))
+		.SetGroup(WorkspaceMenuCategoryRef)
+		.SetIcon(FSlateIcon(FAppStyle::GetAppStyleSetName(), "LevelEditor.Tabs.Details"));
+}
+
+void FMediaSourceManagerEditorToolkit::UnregisterTabSpawners(const TSharedRef<class FTabManager>& InTabManager)
+{
+	FAssetEditorToolkit::UnregisterTabSpawners(InTabManager);
+
+	InTabManager->UnregisterTabSpawner(MediaSourceManagerEditorToolkit::SourcesTabId);
+}
+
+FText FMediaSourceManagerEditorToolkit::GetBaseToolkitName() const
+{
+	return LOCTEXT("AppLabel", "Media Source Manager Editor");
+}
+
+FName FMediaSourceManagerEditorToolkit::GetToolkitFName() const
+{
+	return FName("MediaSourceManagerEditor");
+}
+
+FLinearColor FMediaSourceManagerEditorToolkit::GetWorldCentricTabColorScale() const
+{
+	return FLinearColor(0.3f, 0.2f, 0.5f, 0.5f);
+}
+
+FString FMediaSourceManagerEditorToolkit::GetWorldCentricTabPrefix() const
+{
+	return LOCTEXT("WorldCentricTabPrefix", "MediaSourceManager ").ToString();
+}
+
+void FMediaSourceManagerEditorToolkit::AddReferencedObjects(FReferenceCollector& Collector)
+{
+	Collector.AddReferencedObject(MediaSourceManager);
+}
+
+void FMediaSourceManagerEditorToolkit::PostUndo(bool bSuccess)
+{
+}
+
+void FMediaSourceManagerEditorToolkit::PostRedo(bool bSuccess)
+{
+	PostUndo(bSuccess);
+}
+
+TSharedRef<SDockTab> FMediaSourceManagerEditorToolkit::HandleTabManagerSpawnTab(const FSpawnTabArgs& Args, FName TabIdentifier)
+{
+	TSharedPtr<SWidget> TabWidget = SNullWidget::NullWidget;
+
+	if (TabIdentifier == MediaSourceManagerEditorToolkit::SourcesTabId)
+	{
+		TabWidget = SNew(SMediaSourceManagerSources, *MediaSourceManager);
+	}
+
+	return SNew(SDockTab)
+		.TabRole(ETabRole::PanelTab)
+		[
+			TabWidget.ToSharedRef()
+		];
+}
+
+#undef LOCTEXT_NAMESPACE
diff --git a/Engine/Plugins/Media/MediaSourceManager/Source/MediaSourceManagerEditor/Private/Toolkits/MediaSourceManagerEditorToolkit.h b/Engine/Plugins/Media/MediaSourceManager/Source/MediaSourceManagerEditor/Private/Toolkits/MediaSourceManagerEditorToolkit.h
new file mode 100644
--- /dev/null
+++ b/Engine/Plugins/Media/MediaSourceManager/Source/MediaSourceManagerEditor/Private/Toolkits/MediaSourceManagerEditorToolkit.h
@@ -0,0 +1,70 @@
+// Copyright Epic Games, Inc. All Rights Reserved.
+
+#pragma once
+
+#include "EditorUndoClient.h"
+#include "Framework/Docking/TabManager.h"
+#include "Toolkits/AssetEditorToolkit.h"
+#include "UObject/GCObject.h"
+
+class UMediaSourceManager;
+
+/**
+ * Implements an Editor toolkit for the media source manager.
+ */
+class FMediaSourceManagerEditorToolkit
+	: public FAssetEditorToolkit
+	, public FEditorUndoClient
+	, public FGCObject
+{
+public:
+	/**
+	 * Creates and initializes a new instance.
+	 */
+	FMediaSourceManagerEditorToolkit();
+
+	/** Virtual destructor. */
+	virtual ~FMediaSourceManagerEditorToolkit();
+
+	/**
+	 * Initializes the editor tool kit.
+	 *
+	 * @param InMediaSourceManager The media source manager asset to edit.
+	 * @param InMode The mode to create the toolkit in.
+	 * @param InToolkitHost The toolkit host.
+	 */
+	void Initialize(UMediaSourceManager* InMediaSourceManager, const EToolkitMode::Type InMode,
+		const TSharedPtr<IToolkitHost>& InToolkitHost);
+
+	//~ FAssetEditorToolkit interface
+	virtual FString GetDocumentationLink() const override;
+	virtual void RegisterTabSpawners(const TSharedRef<FTabManager>& TabManager) override;
+	virtual void UnregisterTabSpawners(const TSharedRef<FTabManager>& TabManager) override;
+
+	//~ IToolkit interface
+	virtual FText GetBaseToolkitName() const override;
+	virtual FName GetToolkitFName() const override;
+	virtual FLinearColor GetWorldCentricTabColorScale() const override;
+	virtual FString GetWorldCentricTabPrefix() const override;
+
+	//~ FGCObject interface
+	virtual void AddReferencedObjects(FReferenceCollector& Collector) override;
+	virtual FString GetReferencerName() const override
+	{
+		return TEXT("FMediaSourceManagerEditorToolkit");
+	}
+
+protected:
+
+	//~ FEditorUndoClient interface
+	virtual void PostUndo(bool bSuccess) override;
+	virtual void PostRedo(bool bSuccess) override;
+
+private:
+
+	/** Callback for spawning tabs. */
+	TSharedRef<SDockTab> HandleTabManagerSpawnTab(const FSpawnTabArgs& Args, FName TabIdentifier);
+
+	/** The media source manager asset being edited. */
+	TObjectPtr<UMediaSourceManager> MediaSourceManager;
+};
diff --git a/Engine/Plugins/Media/MediaSourceManager/Source/MediaSourceManagerEditor/Private/Widgets/SMediaSourceManagerSources.cpp b/Engine/Plugins/Media/MediaSourceManager/Source/MediaSourceManagerEditor/Private/Widgets/SMediaSourceManagerSources.cpp
new file mode 100644
--- /dev/null
+++ b/Engine/Plugins/Media/MediaSourceManager/Source/MediaSourceManagerEditor/Private/Widgets/SMediaSourceManagerSources.cpp
@@ -0,0 +1,22 @@
+// Copyright Epic Games, Inc. All Rights Reserved.
+
+#include "Widgets/SMediaSourceManagerSources.h"
+#include "MediaSourceManager.h"
+#include "Widgets/Text/STextBlock.h"
+
+#define LOCTEXT_NAMESPACE "SMediaSourceManagerSources"
+
+void SMediaSourceManagerSources::Construct(const FArguments& InArgs,
+	UMediaSourceManager& InMediaSourceManager)
+{
+	MediaSourceManager = &InMediaSourceManager;
+
+	ChildSlot
+		[
+			SNew(STextBlock)
+				.Text(LOCTEXT("Sources", "Sources"))
+
+		];
+}
+
+#undef LOCTEXT_NAMESPACE
diff --git a/Engine/Plugins/Media/MediaSourceManager/Source/MediaSourceManagerEditor/Private/Widgets/SMediaSourceManagerSources.h b/Engine/Plugins/Media/MediaSourceManager/Source/MediaSourceManagerEditor/Private/Widgets/SMediaSourceManagerSources.h
new file mode 100644
--- /dev/null
+++ b/Engine/Plugins/Media/MediaSourceManager/Source/MediaSourceManagerEditor/Private/Widgets/SMediaSourceManagerSources.h
@@ -0,0 +1,35 @@
+// Copyright Epic Games, Inc. All Rights Reserved.
+
+#pragma once
+
+#include "Widgets/DeclarativeSyntaxSupport.h"
+#include "Widgets/SCompoundWidget.h"
+
+class UMediaSourceManager;
+
+/**
+ * Implements the sources panel of the MediaSourceManager asset editor.
+ */
+class SMediaSourceManagerSources
+	: public SCompoundWidget
+{
+public:
+
+	SLATE_BEGIN_ARGS(SMediaSourceManagerSources) { }
+	SLATE_END_ARGS()
+
+public:
+
+	/**
+	 * Construct this widget.
+	 *
+	 * @param InArgs					The declaration data for this widget.
+	 * @param InMediaSourceManager		The manager to show details for.
+	 */
+	void Construct(const FArguments& InArgs, UMediaSourceManager& InMediaSourceManager);
+
+private:
+
+	/** Pointer to the object that is being viewed. */
+	TWeakObjectPtr<UMediaSourceManager> MediaSourceManager;
+};
