MediaSourceManager: Renamed file.
#rnx

diff --git a/Engine/Plugins/Media/MediaSourceManager/Source/MediaSourceManagerEditor/Private/MediaSourceManagerEditorModule.cpp b/Engine/Plugins/Media/MediaSourceManager/Source/MediaSourceManagerEditor/Private/MediaSourceManagerEditorModule.cpp
new file mode 100644
--- /dev/null
+++ b/Engine/Plugins/Media/MediaSourceManager/Source/MediaSourceManagerEditor/Private/MediaSourceManagerEditorModule.cpp
@@ -0,0 +1,82 @@
+// Copyright Epic Games, Inc. All Rights Reserved.
+
+#include "MediaSourceManagerEditorModule.h"
+
+#include "AssetToolsModule.h"
+#include "AssetTools/MediaSourceManagerActions.h"
+#include "Modules/ModuleManager.h"
+
+#define LOCTEXT_NAMESPACE "MediaSourceManagerEditorModule"
+
+DEFINE_LOG_CATEGORY(LogMediaSourceManagerEditor);
+
+/**
+ * Implements the MediaSourceManagerEditor module.
+ */
+class FMediaSourceManagerEditorModule
+	: public IMediaSourceManagerEditorModule
+{
+public:
+	//~ IMediaSourceManagerEditorModule interface
+	
+
+	//~ IModuleInterface interface
+
+	virtual void StartupModule() override
+	{
+		RegisterAssetTools();
+	}
+
+	virtual void ShutdownModule() override
+	{
+		UnregisterAssetTools();
+	}
+
+protected:
+
+	/**
+	 * Register all our asset tools.
+	 */
+	void RegisterAssetTools()
+	{
+		IAssetTools& AssetTools = FModuleManager::LoadModuleChecked<FAssetToolsModule>("AssetTools").Get();
+
+		RegisterAssetTypeAction(AssetTools, MakeShareable(new FMediaSourceManagerActions()));
+	}
+
+	/**
+	 * Registers a single asset type action.
+	 *
+	 * @param AssetTools	The asset tools object to register with.
+	 * @param Action		The asset type action to register.
+	 */
+	void RegisterAssetTypeAction(IAssetTools& AssetTools, TSharedRef<IAssetTypeActions> Action)
+	{
+		AssetTools.RegisterAssetTypeActions(Action);
+		RegisteredAssetTypeActions.Add(Action);
+	}
+
+	/**
+	 * Unregister all our asset tools.
+	 */
+	void UnregisterAssetTools()
+	{
+		FAssetToolsModule* AssetToolsModule = FModuleManager::GetModulePtr<FAssetToolsModule>("AssetTools");
+
+		if (AssetToolsModule != nullptr)
+		{
+			IAssetTools& AssetTools = AssetToolsModule->Get();
+			for (const TSharedRef<IAssetTypeActions>& Action : RegisteredAssetTypeActions)
+			{
+				AssetTools.UnregisterAssetTypeActions(Action);
+			}
+		}
+	}
+
+	/** The collection of registered asset type actions. */
+	TArray<TSharedRef<IAssetTypeActions>> RegisteredAssetTypeActions;
+};
+
+IMPLEMENT_MODULE(FMediaSourceManagerEditorModule, MediaSourceManagerEditor);
+
+#undef LOCTEXT_NAMESPACE
diff --git a/Engine/Plugins/Media/MediaSourceManager/Source/MediaSourceManagerEditor/Private/MediaSourceManagerModuleEditor.cpp b/Engine/Plugins/Media/MediaSourceManager/Source/MediaSourceManagerEditor/Private/MediaSourceManagerModuleEditor.cpp
deleted file mode 100644
