MediaSourceManager: Added channels and inputs.

#rnx

diff --git a/Engine/Plugins/Media/MediaSourceManager/MediaSourceManager.uplugin b/Engine/Plugins/Media/MediaSourceManager/MediaSourceManager.uplugin
--- a/Engine/Plugins/Media/MediaSourceManager/MediaSourceManager.uplugin
+++ b/Engine/Plugins/Media/MediaSourceManager/MediaSourceManager.uplugin
@@ -15,8 +15,7 @@
 	"IsBetaVersion": false,
 	"IsExperimentalVersion": false,
 	"Installed": false,
-	"Modules": 
-	[
+	"Modules": [
 		{
 			"Name": "MediaSourceManager",
 			"Type": "Runtime",
@@ -27,5 +26,11 @@
 			"Type": "Editor",
 			"LoadingPhase": "Default"
 		}
+	],
+	"Plugins": [
+		{
+			"Name": "MediaFrameworkUtilities",
+			"Enabled": true
+		}
 	]
 }
\ No newline at end of file
diff --git a/Engine/Plugins/Media/MediaSourceManager/Source/MediaSourceManager/MediaSourceManager.Build.cs b/Engine/Plugins/Media/MediaSourceManager/Source/MediaSourceManager/MediaSourceManager.Build.cs
--- a/Engine/Plugins/Media/MediaSourceManager/Source/MediaSourceManager/MediaSourceManager.Build.cs
+++ b/Engine/Plugins/Media/MediaSourceManager/Source/MediaSourceManager/MediaSourceManager.Build.cs
@@ -13,6 +13,8 @@ public class MediaSourceManager : ModuleRules
 			{
 				"Core",
 				"CoreUObject",
+				"MediaAssets",
+				"MediaFrameworkUtilities",
 			}
 			);
 	}
diff --git a/Engine/Plugins/Media/MediaSourceManager/Source/MediaSourceManager/Private/Inputs/MediaSourceManagerInputMediaSource.cpp b/Engine/Plugins/Media/MediaSourceManager/Source/MediaSourceManager/Private/Inputs/MediaSourceManagerInputMediaSource.cpp
new file mode 100644
--- /dev/null
+++ b/Engine/Plugins/Media/MediaSourceManager/Source/MediaSourceManager/Private/Inputs/MediaSourceManagerInputMediaSource.cpp
@@ -0,0 +1,26 @@
+// Copyright Epic Games, Inc. All Rights Reserved.
+
+#include "Inputs/MediaSourceManagerInputMediaSource.h"
+
+#include "MediaSource.h"
+
+#define LOCTEXT_NAMESPACE "MediaSourceManagerInputMediaSource"
+
+FString UMediaSourceManagerInputMediaSource::GetDisplayName()
+{
+	FString Name = FString(TEXT("<none>"));
+
+	if (MediaSource != nullptr)
+	{
+		Name = MediaSource->GetName();
+	}
+
+	return Name;
+}
+
+UMediaSource* UMediaSourceManagerInputMediaSource::GetMediaSource()
+{
+	return MediaSource;
+}
+
+#undef LOCTEXT_NAMESPACE
diff --git a/Engine/Plugins/Media/MediaSourceManager/Source/MediaSourceManager/Private/MediaSourceManager.cpp b/Engine/Plugins/Media/MediaSourceManager/Source/MediaSourceManager/Private/MediaSourceManager.cpp
--- a/Engine/Plugins/Media/MediaSourceManager/Source/MediaSourceManager/Private/MediaSourceManager.cpp
+++ b/Engine/Plugins/Media/MediaSourceManager/Source/MediaSourceManager/Private/MediaSourceManager.cpp
@@ -4,6 +4,18 @@
 
 #define LOCTEXT_NAMESPACE "MediaSourceManager"
 
+void UMediaSourceManager::Validate()
+{
+	// Check we have some channels.
+	int32 NumChannels = 8;
+
+	while (Channels.Num() < NumChannels)
+	{
+		UMediaSourceManagerChannel* Channel = NewObject<UMediaSourceManagerChannel>(this);
+		Channel->Name = FString::FromInt(Channels.Num());
+		Channels.Add(Channel);
+	}
+}
 
 
 #undef LOCTEXT_NAMESPACE
diff --git a/Engine/Plugins/Media/MediaSourceManager/Source/MediaSourceManager/Private/MediaSourceManagerChannel.cpp b/Engine/Plugins/Media/MediaSourceManager/Source/MediaSourceManager/Private/MediaSourceManagerChannel.cpp
new file mode 100644
--- /dev/null
+++ b/Engine/Plugins/Media/MediaSourceManager/Source/MediaSourceManager/Private/MediaSourceManagerChannel.cpp
@@ -0,0 +1,28 @@
+// Copyright Epic Games, Inc. All Rights Reserved.
+
+#include "MediaSourceManagerChannel.h"
+
+#include "MediaAssets/ProxyMediaSource.h"
+
+#define LOCTEXT_NAMESPACE "MediaSourceManagerChannel"
+
+#if WITH_EDITOR
+
+void UMediaSourceManagerChannel::PostEditChangeProperty(FPropertyChangedEvent& PropertyChangedEvent)
+{
+	// Did the input media source change?
+	if (PropertyChangedEvent.GetPropertyName() == GET_MEMBER_NAME_CHECKED(ThisClass, InMediaSource))
+	{
+		// Hook up the proxy media source to this.
+		if (OutMediaSource != nullptr)
+		{
+			OutMediaSource->SetMediaSource(InMediaSource);
+			OutMediaSource->MarkPackageDirty();
+		}
+	}
+}
+
+#endif // WITH_EDITOR
+
+
+#undef LOCTEXT_NAMESPACE
diff --git a/Engine/Plugins/Media/MediaSourceManager/Source/MediaSourceManager/Public/Inputs/MediaSourceManagerInput.h b/Engine/Plugins/Media/MediaSourceManager/Source/MediaSourceManager/Public/Inputs/MediaSourceManagerInput.h
new file mode 100644
--- /dev/null
+++ b/Engine/Plugins/Media/MediaSourceManager/Source/MediaSourceManager/Public/Inputs/MediaSourceManagerInput.h
@@ -0,0 +1,36 @@
+// Copyright Epic Games, Inc. All Rights Reserved.
+
+#pragma once
+
+#include "CoreTypes.h"
+#include "UObject/Object.h"
+#include "UObject/ObjectMacros.h"
+
+#include "MediaSourceManagerInput.generated.h"
+
+class UMediaSource;
+
+/**
+* Base class for an input for the media source manager.
+* This provides a source from capture cards, movie files, etc.
+*/
+UCLASS(Abstract, BlueprintType)
+class MEDIASOURCEMANAGER_API UMediaSourceManagerInput : public UObject
+{
+	GENERATED_BODY()
+
+public:
+	/**
+	 * Get the name of this input.
+	 */
+	UFUNCTION(BlueprintCallable, Category = "Media|MediaSourceManager|Input")
+	virtual FString GetDisplayName() PURE_VIRTUAL(UMediaSourceManagerInput::GetDisplayName, return FString(););
+
+	/**
+	 * Get a media source for this input.
+	 */
+	UFUNCTION(BlueprintCallable, Category = "Media|MediaSourceManager|Input")
+	virtual UMediaSource* GetMediaSource() PURE_VIRTUAL(UMediaSourceManagerInput::GetMediaSource, return nullptr;);
+
+};
+
diff --git a/Engine/Plugins/Media/MediaSourceManager/Source/MediaSourceManager/Public/Inputs/MediaSourceManagerInputMediaSource.h b/Engine/Plugins/Media/MediaSourceManager/Source/MediaSourceManager/Public/Inputs/MediaSourceManagerInputMediaSource.h
new file mode 100644
--- /dev/null
+++ b/Engine/Plugins/Media/MediaSourceManager/Source/MediaSourceManager/Public/Inputs/MediaSourceManagerInputMediaSource.h
@@ -0,0 +1,22 @@
+// Copyright Epic Games, Inc. All Rights Reserved.
+
+#pragma once
+
+#include "MediaSourceManagerInput.h"
+
+#include "MediaSourceManagerInputMediaSource.generated.h"
+
+UCLASS(BlueprintType)
+class MEDIASOURCEMANAGER_API UMediaSourceManagerInputMediaSource : public UMediaSourceManagerInput
+{
+	GENERATED_BODY()
+
+public:
+	virtual FString GetDisplayName() override;
+	virtual UMediaSource* GetMediaSource() override;
+
+	/** The media source this input uses. */
+	UPROPERTY(EditAnywhere, BlueprintReadWrite, Category = "InputMediaSource")
+	TObjectPtr<UMediaSource> MediaSource = nullptr;
+};
+
diff --git a/Engine/Plugins/Media/MediaSourceManager/Source/MediaSourceManager/Public/MediaSourceManager.h b/Engine/Plugins/Media/MediaSourceManager/Source/MediaSourceManager/Public/MediaSourceManager.h
--- a/Engine/Plugins/Media/MediaSourceManager/Source/MediaSourceManager/Public/MediaSourceManager.h
+++ b/Engine/Plugins/Media/MediaSourceManager/Source/MediaSourceManager/Public/MediaSourceManager.h
@@ -3,6 +3,7 @@
 #pragma once
 
 #include "CoreTypes.h"
+#include "MediaSourceManagerChannel.h"
 #include "UObject/Object.h"
 #include "UObject/ObjectMacros.h"
 
@@ -11,11 +12,19 @@
 /**
 * Manager to handle media sources and their connections.
 */
-UCLASS()
+UCLASS(BlueprintType)
 class MEDIASOURCEMANAGER_API UMediaSourceManager : public UObject
 {
 	GENERATED_BODY()
 
 public:
+	/** Our channels. */
+	UPROPERTY(EditAnywhere, BlueprintReadWrite, Category = "Channels")
+	TArray<UMediaSourceManagerChannel*> Channels;
+
+	/**
+	 * Call this to make sure everything is set up.
+	 */
+	void Validate();
 
 };
diff --git a/Engine/Plugins/Media/MediaSourceManager/Source/MediaSourceManager/Public/MediaSourceManagerChannel.h b/Engine/Plugins/Media/MediaSourceManager/Source/MediaSourceManager/Public/MediaSourceManagerChannel.h
new file mode 100644
--- /dev/null
+++ b/Engine/Plugins/Media/MediaSourceManager/Source/MediaSourceManager/Public/MediaSourceManagerChannel.h
@@ -0,0 +1,45 @@
+// Copyright Epic Games, Inc. All Rights Reserved.
+
+#pragma once
+
+#include "CoreTypes.h"
+#include "UObject/Object.h"
+#include "UObject/ObjectMacros.h"
+
+#include "MediaSourceManagerChannel.generated.h"
+
+class UMediaSource;
+class UMediaSourceManagerInput;
+class UProxyMediaSource;
+
+/**
+* Handles a single channel for the MediaSourceManager.
+*/
+UCLASS(BlueprintType)
+class MEDIASOURCEMANAGER_API UMediaSourceManagerChannel : public UObject
+{
+	GENERATED_BODY()
+
+public:
+	/** The name of this channel. */
+	UPROPERTY(EditAnywhere, BlueprintReadWrite, Category = "Channel")
+	FString Name;
+
+	/** Our input. */
+	UPROPERTY(EditAnywhere, BlueprintReadWrite, Category = "Channel")
+	TObjectPtr<UMediaSourceManagerInput> Input = nullptr;
+
+	/** Our media sources for our inputs. */
+	TObjectPtr<UMediaSource> InMediaSource = nullptr;
+
+	/** Our persistent media source. */
+	UPROPERTY(EditAnywhere, BlueprintReadWrite, Category = "Channel")
+	TObjectPtr<UProxyMediaSource> OutMediaSource = nullptr;
+
+	//~ Begin UObject interface
+#if WITH_EDITOR
+	virtual void PostEditChangeProperty(FPropertyChangedEvent& PropertyChangedEvent) override;
+#endif // WITH_EDITOR
+	//~ End UObject interface
+
+};
diff --git a/Engine/Plugins/Media/MediaSourceManager/Source/MediaSourceManagerEditor/MediaSourceManagerEditor.Build.cs b/Engine/Plugins/Media/MediaSourceManager/Source/MediaSourceManagerEditor/MediaSourceManagerEditor.Build.cs
--- a/Engine/Plugins/Media/MediaSourceManager/Source/MediaSourceManagerEditor/MediaSourceManagerEditor.Build.cs
+++ b/Engine/Plugins/Media/MediaSourceManager/Source/MediaSourceManagerEditor/MediaSourceManagerEditor.Build.cs
@@ -13,6 +13,7 @@ public class MediaSourceManagerEditor : ModuleRules
 			{
 				"Core",
 				"CoreUObject",
+				"MediaAssets",
 				"MediaSourceManager",
 				"Slate",
 				"SlateCore",
diff --git a/Engine/Plugins/Media/MediaSourceManager/Source/MediaSourceManagerEditor/Private/Factories/MediaSourceManagerFactory.cpp b/Engine/Plugins/Media/MediaSourceManager/Source/MediaSourceManagerEditor/Private/Factories/MediaSourceManagerFactory.cpp
--- a/Engine/Plugins/Media/MediaSourceManager/Source/MediaSourceManagerEditor/Private/Factories/MediaSourceManagerFactory.cpp
+++ b/Engine/Plugins/Media/MediaSourceManager/Source/MediaSourceManagerEditor/Private/Factories/MediaSourceManagerFactory.cpp
@@ -21,7 +21,9 @@ UMediaSourceManagerFactory::UMediaSourceManagerFactory(const FObjectInitializer&
 
 UObject* UMediaSourceManagerFactory::FactoryCreateNew(UClass* InClass, UObject* InParent, FName InName, EObjectFlags Flags, UObject* Context, FFeedbackContext* Warn)
 {
-	return NewObject<UMediaSourceManager>(InParent, InClass, InName, Flags);
+	UMediaSourceManager* Manager = NewObject<UMediaSourceManager>(InParent, InClass, InName, Flags);
+	Manager->Validate();
+	return Manager;
 }
 
 uint32 UMediaSourceManagerFactory::GetMenuCategories() const
diff --git a/Engine/Plugins/Media/MediaSourceManager/Source/MediaSourceManagerEditor/Private/Widgets/SMediaSourceManagerChannel.cpp b/Engine/Plugins/Media/MediaSourceManager/Source/MediaSourceManagerEditor/Private/Widgets/SMediaSourceManagerChannel.cpp
new file mode 100644
--- /dev/null
+++ b/Engine/Plugins/Media/MediaSourceManager/Source/MediaSourceManagerEditor/Private/Widgets/SMediaSourceManagerChannel.cpp
@@ -0,0 +1,116 @@
+// Copyright Epic Games, Inc. All Rights Reserved.
+
+#include "Widgets/SMediaSourceManagerChannel.h"
+
+#include "DragAndDrop/AssetDragDropOp.h"
+#include "Inputs/MediaSourceManagerInputMediaSource.h"
+#include "MediaSource.h"
+#include "MediaSourceManagerChannel.h"
+#include "Widgets/Layout/SScrollBox.h"
+#include "Widgets/Text/STextBlock.h"
+
+#define LOCTEXT_NAMESPACE "SMediaSourceManagerChannel"
+
+void SMediaSourceManagerChannel::Construct(const FArguments& InArgs,
+	UMediaSourceManagerChannel* InChannel)
+{
+	ChannelPtr = InChannel;
+
+	ChildSlot
+		[
+			SNew(SHorizontalBox)
+
+			// Name of channel.
+			+ SHorizontalBox::Slot()
+				.FillWidth(0.11f)
+				.Padding(2)
+				.HAlign(HAlign_Left)
+				[
+					SNew(STextBlock)
+						.Text(FText::FromString(ChannelPtr->Name))
+				]
+
+			// Name of input.
+			+ SHorizontalBox::Slot()
+				.FillWidth(0.11f)
+				.Padding(2)
+				.HAlign(HAlign_Left)
+				[
+					SAssignNew(InputNameTextBlock, STextBlock)
+				]
+		];
+
+	Refresh();
+}
+
+void SMediaSourceManagerChannel::OnDragEnter(const FGeometry& MyGeometry, const FDragDropEvent& DragDropEvent)
+{
+}
+
+void SMediaSourceManagerChannel::OnDragLeave(const FDragDropEvent& DragDropEvent)
+{
+}
+
+FReply SMediaSourceManagerChannel::OnDragOver(const FGeometry& MyGeometry, const FDragDropEvent& DragDropEvent)
+{
+	// Is this an asset drop?
+	if (TSharedPtr<FAssetDragDropOp> AssetDragDrop = DragDropEvent.GetOperationAs<FAssetDragDropOp>())
+	{
+		return FReply::Handled();
+	}
+
+	return FReply::Unhandled();
+}
+
+FReply SMediaSourceManagerChannel::OnDrop(const FGeometry& MyGeometry, const FDragDropEvent& DragDropEvent)
+{
+	// Is this an asset drop?
+	if (TSharedPtr<FAssetDragDropOp> AssetDragDrop = DragDropEvent.GetOperationAs<FAssetDragDropOp>())
+	{
+		UMediaSourceManagerChannel* Channel = ChannelPtr.Get();
+		if (Channel != nullptr)
+		{
+			for (const FAssetData& Asset : AssetDragDrop->GetAssets())
+			{
+				// Is this a media source?
+				UMediaSource* MediaSource = Cast<UMediaSource>(Asset.GetAsset());
+				if (MediaSource != nullptr)
+				{
+					Channel->Modify();
+					UMediaSourceManagerInputMediaSource* Input = NewObject<UMediaSourceManagerInputMediaSource>(Channel);
+					Input->MediaSource = MediaSource;
+					Channel->Input = Input;
+						
+					Refresh();
+					break;
+				}
+			}
+		}
+		return FReply::Handled();
+	}
+
+	return FReply::Unhandled();
+}
+
+void SMediaSourceManagerChannel::Refresh()
+{
+	// Get channel.
+	UMediaSourceManagerChannel* Channel = ChannelPtr.Get();
+	if (Channel != nullptr)
+	{
+		FString InputName = FString(TEXT("None"));
+
+		// Get input.
+		UMediaSourceManagerInput* Input = Channel->Input;
+		if (Input != nullptr)
+		{
+			InputName = Input->GetDisplayName();
+		}
+
+		// Update input widgets.
+		InputNameTextBlock->SetText(FText::FromString(InputName));
+	}
+}
+
+
+#undef LOCTEXT_NAMESPACE
diff --git a/Engine/Plugins/Media/MediaSourceManager/Source/MediaSourceManagerEditor/Private/Widgets/SMediaSourceManagerChannel.h b/Engine/Plugins/Media/MediaSourceManager/Source/MediaSourceManagerEditor/Private/Widgets/SMediaSourceManagerChannel.h
new file mode 100644
--- /dev/null
+++ b/Engine/Plugins/Media/MediaSourceManager/Source/MediaSourceManagerEditor/Private/Widgets/SMediaSourceManagerChannel.h
@@ -0,0 +1,50 @@
+// Copyright Epic Games, Inc. All Rights Reserved.
+
+#pragma once
+
+#include "Widgets/DeclarativeSyntaxSupport.h"
+#include "Widgets/SCompoundWidget.h"
+
+class SVerticalBox;
+class UMediaSourceManagerChannel;
+
+/**
+ * Implements a single channel.
+ */
+class SMediaSourceManagerChannel
+	: public SCompoundWidget
+{
+public:
+
+	SLATE_BEGIN_ARGS(SMediaSourceManagerChannel) { }
+	SLATE_END_ARGS()
+
+	/**
+	 * Construct this widget.
+	 *
+	 * @param InArgs					The declaration data for this widget.
+	 * @param InChannel					The channel to show details for.
+	 */
+	void Construct(const FArguments& InArgs, UMediaSourceManagerChannel* InChannel);
+
+protected:
+	// SWidget interface
+	virtual void OnDragEnter(const FGeometry& MyGeometry, const FDragDropEvent& DragDropEvent) override;
+	virtual void OnDragLeave(const FDragDropEvent& DragDropEvent) override;
+	virtual FReply OnDragOver(const FGeometry& MyGeometry, const FDragDropEvent& DragDropEvent) override;
+	virtual FReply OnDrop(const FGeometry& MyGeometry, const FDragDropEvent& DragDropEvent) override;
+	// End of SWidget interface
+
+private:
+
+	/** Pointer to the object that is being viewed. */
+	TWeakObjectPtr<UMediaSourceManagerChannel> ChannelPtr;
+	/** Widget for the input name. */
+	TSharedPtr<STextBlock> InputNameTextBlock;
+	
+	/**
+	 * Refreshes the widgets based on the current sstate.
+	 */
+	void Refresh();
+
+};
diff --git a/Engine/Plugins/Media/MediaSourceManager/Source/MediaSourceManagerEditor/Private/Widgets/SMediaSourceManagerSources.cpp b/Engine/Plugins/Media/MediaSourceManager/Source/MediaSourceManagerEditor/Private/Widgets/SMediaSourceManagerSources.cpp
--- a/Engine/Plugins/Media/MediaSourceManager/Source/MediaSourceManagerEditor/Private/Widgets/SMediaSourceManagerSources.cpp
+++ b/Engine/Plugins/Media/MediaSourceManager/Source/MediaSourceManagerEditor/Private/Widgets/SMediaSourceManagerSources.cpp
@@ -1,7 +1,11 @@
 // Copyright Epic Games, Inc. All Rights Reserved.
 
 #include "Widgets/SMediaSourceManagerSources.h"
+
+#include "Inputs/MediaSourceManagerInput.h"
 #include "MediaSourceManager.h"
+#include "Widgets/Layout/SScrollBox.h"
+#include "Widgets/SMediaSourceManagerChannel.h"
 #include "Widgets/Text/STextBlock.h"
 
 #define LOCTEXT_NAMESPACE "SMediaSourceManagerSources"
@@ -9,14 +13,44 @@
 void SMediaSourceManagerSources::Construct(const FArguments& InArgs,
 	UMediaSourceManager& InMediaSourceManager)
 {
-	MediaSourceManager = &InMediaSourceManager;
+	MediaSourceManagerPtr = &InMediaSourceManager;
 
 	ChildSlot
 		[
-			SNew(STextBlock)
-				.Text(LOCTEXT("Sources", "Sources"))
+			SNew(SScrollBox)
+
+			
+			+ SScrollBox::Slot()
+			[
+				SAssignNew(ChannelsContainer, SVerticalBox)
+			]
 
 		];
+
+	RefreshChannels();
 }
 
+void SMediaSourceManagerSources::RefreshChannels()
+{
+	ChannelsContainer->ClearChildren();
+	ChannelWidgets.Reset();
+	UMediaSourceManager* MediaSourceManager = MediaSourceManagerPtr.Get();
+	if (MediaSourceManager != nullptr)
+	{
+		int32 NumChannels = MediaSourceManager->Channels.Num();
+		ChannelWidgets.SetNum(NumChannels);
+		for (int32 Index = 0; Index < NumChannels; ++Index)
+		{
+			ChannelsContainer->AddSlot()
+				.AutoHeight()
+				.Padding(2)
+				.VAlign(VAlign_Center)
+				[
+					SNew(SMediaSourceManagerChannel, MediaSourceManager->Channels[Index])
+				];
+		}
+	}
+}
+
+
 #undef LOCTEXT_NAMESPACE
diff --git a/Engine/Plugins/Media/MediaSourceManager/Source/MediaSourceManagerEditor/Private/Widgets/SMediaSourceManagerSources.h b/Engine/Plugins/Media/MediaSourceManager/Source/MediaSourceManagerEditor/Private/Widgets/SMediaSourceManagerSources.h
--- a/Engine/Plugins/Media/MediaSourceManager/Source/MediaSourceManagerEditor/Private/Widgets/SMediaSourceManagerSources.h
+++ b/Engine/Plugins/Media/MediaSourceManager/Source/MediaSourceManagerEditor/Private/Widgets/SMediaSourceManagerSources.h
@@ -5,6 +5,7 @@
 #include "Widgets/DeclarativeSyntaxSupport.h"
 #include "Widgets/SCompoundWidget.h"
 
+class SVerticalBox;
 class UMediaSourceManager;
 
 /**
@@ -18,8 +19,6 @@ public:
 	SLATE_BEGIN_ARGS(SMediaSourceManagerSources) { }
 	SLATE_END_ARGS()
 
-public:
-
 	/**
 	 * Construct this widget.
 	 *
@@ -31,5 +30,16 @@ public:
 private:
 
 	/** Pointer to the object that is being viewed. */
-	TWeakObjectPtr<UMediaSourceManager> MediaSourceManager;
+	TWeakObjectPtr<UMediaSourceManager> MediaSourceManagerPtr;
+
+	/** Pointer to the container for the media sources. */
+	TSharedPtr<SVerticalBox> ChannelsContainer;
+
+	struct FChannelWidgets
+	{
+		TSharedPtr<STextBlock> InputName;
+	};
+	TArray<FChannelWidgets> ChannelWidgets;
+
+	void RefreshChannels();
 };
