MediaSourceManager: Added texture output.

#rnx

diff --git a/Engine/Plugins/Media/MediaSourceManager/Source/MediaSourceManager/MediaSourceManager.Build.cs b/Engine/Plugins/Media/MediaSourceManager/Source/MediaSourceManager/MediaSourceManager.Build.cs
--- a/Engine/Plugins/Media/MediaSourceManager/Source/MediaSourceManager/MediaSourceManager.Build.cs
+++ b/Engine/Plugins/Media/MediaSourceManager/Source/MediaSourceManager/MediaSourceManager.Build.cs
@@ -13,6 +13,7 @@ public class MediaSourceManager : ModuleRules
 			{
 				"Core",
 				"CoreUObject",
+				"Engine",
 				"MediaAssets",
 				"MediaFrameworkUtilities",
 			}
diff --git a/Engine/Plugins/Media/MediaSourceManager/Source/MediaSourceManager/Private/MediaSourceManager.cpp b/Engine/Plugins/Media/MediaSourceManager/Source/MediaSourceManager/Private/MediaSourceManager.cpp
--- a/Engine/Plugins/Media/MediaSourceManager/Source/MediaSourceManager/Private/MediaSourceManager.cpp
+++ b/Engine/Plugins/Media/MediaSourceManager/Source/MediaSourceManager/Private/MediaSourceManager.cpp
@@ -15,6 +15,15 @@ void UMediaSourceManager::Validate()
 		Channel->Name = FString::FromInt(Channels.Num());
 		Channels.Add(Channel);
 	}
+
+	// Validate each channel.
+	for (UMediaSourceManagerChannel* Channel : Channels)
+	{
+		if (Channel != nullptr)
+		{
+			Channel->Validate();
+		}
+	}
 }
 
 
diff --git a/Engine/Plugins/Media/MediaSourceManager/Source/MediaSourceManager/Private/MediaSourceManagerChannel.cpp b/Engine/Plugins/Media/MediaSourceManager/Source/MediaSourceManager/Private/MediaSourceManagerChannel.cpp
--- a/Engine/Plugins/Media/MediaSourceManager/Source/MediaSourceManager/Private/MediaSourceManagerChannel.cpp
+++ b/Engine/Plugins/Media/MediaSourceManager/Source/MediaSourceManager/Private/MediaSourceManagerChannel.cpp
@@ -3,9 +3,21 @@
 #include "MediaSourceManagerChannel.h"
 
 #include "MediaAssets/ProxyMediaSource.h"
+#include "MediaTexture.h"
 
 #define LOCTEXT_NAMESPACE "MediaSourceManagerChannel"
 
+void UMediaSourceManagerChannel::Validate()
+{
+	// Make sure we have an out texture.
+	if (OutTexture == nullptr)
+	{
+		Modify();
+		OutTexture = NewObject<UMediaTexture>(this);
+		OutTexture->UpdateResource();
+	}
+}
+
 #if WITH_EDITOR
 
 void UMediaSourceManagerChannel::PostEditChangeProperty(FPropertyChangedEvent& PropertyChangedEvent)
diff --git a/Engine/Plugins/Media/MediaSourceManager/Source/MediaSourceManager/Public/MediaSourceManagerChannel.h b/Engine/Plugins/Media/MediaSourceManager/Source/MediaSourceManager/Public/MediaSourceManagerChannel.h
--- a/Engine/Plugins/Media/MediaSourceManager/Source/MediaSourceManager/Public/MediaSourceManagerChannel.h
+++ b/Engine/Plugins/Media/MediaSourceManager/Source/MediaSourceManager/Public/MediaSourceManagerChannel.h
@@ -11,6 +11,7 @@
 class UMediaSource;
 class UMediaSourceManagerInput;
 class UProxyMediaSource;
+class UTexture;
 
 /**
 * Handles a single channel for the MediaSourceManager.
@@ -36,6 +37,15 @@ public:
 	UPROPERTY(EditAnywhere, BlueprintReadWrite, Category = "Channel")
 	TObjectPtr<UProxyMediaSource> OutMediaSource = nullptr;
 
+	/** The channel will output the media here. */
+	UPROPERTY(VisibleAnywhere, BlueprintReadOnly, Category = "Channel")
+	TObjectPtr<UTexture> OutTexture;
+
+	/**
+	 * Call this to make sure everything is set up.
+	 */
+	void Validate();
+
 	//~ Begin UObject interface
 #if WITH_EDITOR
 	virtual void PostEditChangeProperty(FPropertyChangedEvent& PropertyChangedEvent) override;
diff --git a/Engine/Plugins/Media/MediaSourceManager/Source/MediaSourceManagerEditor/MediaSourceManagerEditor.Build.cs b/Engine/Plugins/Media/MediaSourceManager/Source/MediaSourceManagerEditor/MediaSourceManagerEditor.Build.cs
--- a/Engine/Plugins/Media/MediaSourceManager/Source/MediaSourceManagerEditor/MediaSourceManagerEditor.Build.cs
+++ b/Engine/Plugins/Media/MediaSourceManager/Source/MediaSourceManagerEditor/MediaSourceManagerEditor.Build.cs
@@ -13,6 +13,7 @@ public class MediaSourceManagerEditor : ModuleRules
 			{
 				"Core",
 				"CoreUObject",
+				"InputCore",
 				"MediaAssets",
 				"MediaSourceManager",
 				"Slate",
diff --git a/Engine/Plugins/Media/MediaSourceManager/Source/MediaSourceManagerEditor/Private/Widgets/SMediaSourceManagerChannel.cpp b/Engine/Plugins/Media/MediaSourceManager/Source/MediaSourceManagerEditor/Private/Widgets/SMediaSourceManagerChannel.cpp
--- a/Engine/Plugins/Media/MediaSourceManager/Source/MediaSourceManagerEditor/Private/Widgets/SMediaSourceManagerChannel.cpp
+++ b/Engine/Plugins/Media/MediaSourceManager/Source/MediaSourceManagerEditor/Private/Widgets/SMediaSourceManagerChannel.cpp
@@ -7,6 +7,7 @@
 #include "MediaSource.h"
 #include "MediaSourceManagerChannel.h"
 #include "Widgets/Layout/SScrollBox.h"
+#include "Widgets/SMediaSourceManagerTexture.h"
 #include "Widgets/Text/STextBlock.h"
 
 #define LOCTEXT_NAMESPACE "SMediaSourceManagerChannel"
@@ -38,6 +39,15 @@ void SMediaSourceManagerChannel::Construct(const FArguments& InArgs,
 				[
 					SAssignNew(InputNameTextBlock, STextBlock)
 				]
+
+			// Out texture
+			+ SHorizontalBox::Slot()
+				.FillWidth(0.11f)
+				.Padding(2)
+				.HAlign(HAlign_Left)
+				[
+					SNew(SMediaSourceManagerTexture, ChannelPtr.Get())
+				]
 		];
 
 	Refresh();
diff --git a/Engine/Plugins/Media/MediaSourceManager/Source/MediaSourceManagerEditor/Private/Widgets/SMediaSourceManagerChannel.h b/Engine/Plugins/Media/MediaSourceManager/Source/MediaSourceManagerEditor/Private/Widgets/SMediaSourceManagerChannel.h
--- a/Engine/Plugins/Media/MediaSourceManager/Source/MediaSourceManagerEditor/Private/Widgets/SMediaSourceManagerChannel.h
+++ b/Engine/Plugins/Media/MediaSourceManager/Source/MediaSourceManagerEditor/Private/Widgets/SMediaSourceManagerChannel.h
@@ -5,7 +5,7 @@
 #include "Widgets/DeclarativeSyntaxSupport.h"
 #include "Widgets/SCompoundWidget.h"
 
-class SVerticalBox;
+class STextBlock;
 class UMediaSourceManagerChannel;
 
 /**
diff --git a/Engine/Plugins/Media/MediaSourceManager/Source/MediaSourceManagerEditor/Private/Widgets/SMediaSourceManagerTexture.cpp b/Engine/Plugins/Media/MediaSourceManager/Source/MediaSourceManagerEditor/Private/Widgets/SMediaSourceManagerTexture.cpp
new file mode 100644
--- /dev/null
+++ b/Engine/Plugins/Media/MediaSourceManager/Source/MediaSourceManagerEditor/Private/Widgets/SMediaSourceManagerTexture.cpp
@@ -0,0 +1,57 @@
+// Copyright Epic Games, Inc. All Rights Reserved.
+
+#include "Widgets/SMediaSourceManagerTexture.h"
+
+#include "AssetRegistry/AssetData.h"
+#include "AssetThumbnail.h"
+#include "DragAndDrop/AssetDragDropOp.h"
+#include "Engine/Texture.h"
+#include "MediaSourceManagerChannel.h"
+#include "ThumbnailRendering/ThumbnailManager.h"
+
+void SMediaSourceManagerTexture::Construct(const FArguments& InArgs,
+	UMediaSourceManagerChannel* InChannel)
+{
+	ChannelPtr = InChannel;
+
+	// Set up texture thumbnail.
+	UTexture* Texture = nullptr;
+	if (ChannelPtr.IsValid())
+	{
+		Texture = ChannelPtr->OutTexture;
+	}
+	FAssetData ThumbnailAssetData = Texture;
+	TSharedPtr<FAssetThumbnail> Thumbnail = MakeShareable(new FAssetThumbnail(ThumbnailAssetData,
+		64, 64, UThumbnailManager::Get().GetSharedThumbnailPool()));
+
+	ChildSlot
+		[
+			Thumbnail->MakeThumbnailWidget()
+		];
+}
+
+FReply SMediaSourceManagerTexture::OnMouseButtonDown(const FGeometry& MyGeometry, const FPointerEvent& MouseEvent)
+{
+	if (MouseEvent.GetEffectingButton() == EKeys::LeftMouseButton)
+	{
+		return FReply::Handled().DetectDrag(SharedThis(this), EKeys::LeftMouseButton);
+	}
+
+	return FReply::Unhandled();
+}
+
+FReply SMediaSourceManagerTexture::OnDragDetected(const FGeometry& MyGeometry, const FPointerEvent& MouseEvent)
+{
+	UMediaSourceManagerChannel* Channel = ChannelPtr.Get();
+	if (Channel != nullptr)
+	{
+		UTexture* Texture = Channel->OutTexture;
+		if (Texture != nullptr)
+		{
+			FAssetData AssetData(Texture);
+			return FReply::Handled().BeginDragDrop(FAssetDragDropOp::New(AssetData));
+		}
+	}
+
+	return FReply::Unhandled();
+}
diff --git a/Engine/Plugins/Media/MediaSourceManager/Source/MediaSourceManagerEditor/Private/Widgets/SMediaSourceManagerTexture.h b/Engine/Plugins/Media/MediaSourceManager/Source/MediaSourceManagerEditor/Private/Widgets/SMediaSourceManagerTexture.h
new file mode 100644
--- /dev/null
+++ b/Engine/Plugins/Media/MediaSourceManager/Source/MediaSourceManagerEditor/Private/Widgets/SMediaSourceManagerTexture.h
@@ -0,0 +1,39 @@
+// Copyright Epic Games, Inc. All Rights Reserved.
+
+#pragma once
+
+#include "Widgets/DeclarativeSyntaxSupport.h"
+#include "Widgets/SCompoundWidget.h"
+
+class UMediaSourceManagerChannel;
+
+/**
+ * Handles a single texture.
+ */
+class SMediaSourceManagerTexture
+	: public SCompoundWidget
+{
+public:
+
+	SLATE_BEGIN_ARGS(SMediaSourceManagerTexture) { }
+	SLATE_END_ARGS()
+
+	/**
+	 * Construct this widget.
+	 *
+	 * @param InArgs					The declaration data for this widget.
+	 * @param InChannel					The channel to use.
+	 */
+	void Construct(const FArguments& InArgs, UMediaSourceManagerChannel* InChannel);
+
+protected:
+	// SWidget interface
+	virtual FReply OnMouseButtonDown(const FGeometry& MyGeometry, const FPointerEvent& MouseEvent) override;
+	virtual FReply OnDragDetected(const FGeometry& MyGeometry, const FPointerEvent& MouseEvent) override;
+	// End of SWidget interface
+
+private:
+
+	/** Pointer to the object that is being viewed. */
+	TWeakObjectPtr<UMediaSourceManagerChannel> ChannelPtr;
+};
