MediaSourceManager: Non unity build fix.
#rnx

diff --git a/Engine/Plugins/Media/MediaSourceManager/Source/MediaSourceManagerEditor/Private/Widgets/SMediaSourceManagerSources.h b/Engine/Plugins/Media/MediaSourceManager/Source/MediaSourceManagerEditor/Private/Widgets/SMediaSourceManagerSources.h
--- a/Engine/Plugins/Media/MediaSourceManager/Source/MediaSourceManagerEditor/Private/Widgets/SMediaSourceManagerSources.h
+++ b/Engine/Plugins/Media/MediaSourceManager/Source/MediaSourceManagerEditor/Private/Widgets/SMediaSourceManagerSources.h
@@ -5,6 +5,7 @@
 #include "Widgets/DeclarativeSyntaxSupport.h"
 #include "Widgets/SCompoundWidget.h"
 
+class STextBlock;
 class SVerticalBox;
 class UMediaSourceManager;
 
