MediaSourceManager: Added a player and preview window for each channel.

#rnx

diff --git a/Engine/Plugins/Media/MediaSourceManager/Source/MediaSourceManager/Private/MediaSourceManagerChannel.cpp b/Engine/Plugins/Media/MediaSourceManager/Source/MediaSourceManager/Private/MediaSourceManagerChannel.cpp
--- a/Engine/Plugins/Media/MediaSourceManager/Source/MediaSourceManager/Private/MediaSourceManagerChannel.cpp
+++ b/Engine/Plugins/Media/MediaSourceManager/Source/MediaSourceManager/Private/MediaSourceManagerChannel.cpp
@@ -2,7 +2,9 @@
 
 #include "MediaSourceManagerChannel.h"
 
+#include "Inputs/MediaSourceManagerInput.h"
 #include "MediaAssets/ProxyMediaSource.h"
+#include "MediaPlayer.h"
 #include "MediaTexture.h"
 
 #define LOCTEXT_NAMESPACE "MediaSourceManagerChannel"
@@ -13,8 +15,51 @@ void UMediaSourceManagerChannel::Validate()
 	if (OutTexture == nullptr)
 	{
 		Modify();
-		OutTexture = NewObject<UMediaTexture>(this);
-		OutTexture->UpdateResource();
+		UMediaTexture* MediaTexture = NewObject<UMediaTexture>(this);
+		MediaTexture->NewStyleOutput = true;
+		MediaTexture->UpdateResource();
+		OutTexture = MediaTexture;
+	}
+}
+
+UMediaPlayer* UMediaSourceManagerChannel::GetMediaPlayer()
+{
+	// Create a player if we don't have one.
+	if (CurrentMediaPlayer == nullptr)
+	{
+		CurrentMediaPlayer = NewObject<UMediaPlayer>(this, "MediaPlayer", RF_Transient);
+		CurrentMediaPlayer->SetLooping(false);
+		CurrentMediaPlayer->PlayOnOpen = false;
+		if (OutTexture != nullptr)
+		{
+			UMediaTexture* MediaTexture = Cast<UMediaTexture>(OutTexture);
+			if (MediaTexture != nullptr)
+			{
+				MediaTexture->SetMediaPlayer(CurrentMediaPlayer);
+				MediaTexture->UpdateResource();
+			}
+		}
+	}
+
+	return CurrentMediaPlayer;
+}
+
+void UMediaSourceManagerChannel::Play()
+{
+	if (Input != nullptr)
+	{
+		UMediaSource* MediaSource = Input->GetMediaSource();
+		if (MediaSource != nullptr)
+		{
+			UMediaPlayer* MediaPlayer = GetMediaPlayer();
+			if (MediaPlayer != nullptr)
+			{
+				FMediaPlayerOptions Options;
+				Options.PlayOnOpen = EMediaPlayerOptionBooleanOverride::Enabled;;
+				Options.Loop = EMediaPlayerOptionBooleanOverride::Enabled;
+				MediaPlayer->OpenSourceWithOptions(MediaSource, Options);
+			}
+		}
 	}
 }
 
diff --git a/Engine/Plugins/Media/MediaSourceManager/Source/MediaSourceManager/Public/MediaSourceManagerChannel.h b/Engine/Plugins/Media/MediaSourceManager/Source/MediaSourceManager/Public/MediaSourceManagerChannel.h
--- a/Engine/Plugins/Media/MediaSourceManager/Source/MediaSourceManager/Public/MediaSourceManagerChannel.h
+++ b/Engine/Plugins/Media/MediaSourceManager/Source/MediaSourceManager/Public/MediaSourceManagerChannel.h
@@ -8,6 +8,7 @@
 
 #include "MediaSourceManagerChannel.generated.h"
 
+class UMediaPlayer;
 class UMediaSource;
 class UMediaSourceManagerInput;
 class UProxyMediaSource;
@@ -39,7 +40,16 @@ public:
 
 	/** The channel will output the media here. */
 	UPROPERTY(VisibleAnywhere, BlueprintReadOnly, Category = "Channel")
-	TObjectPtr<UTexture> OutTexture;
+	TObjectPtr<UTexture> OutTexture = nullptr;
+
+	/**
+	 * Call this to get the media player.
+	 */
+	UMediaPlayer* GetMediaPlayer();
+
+	/** Start playback. */
+	UFUNCTION(BlueprintCallable, Category = "Media|MediaSourceManager")
+	void Play();
 
 	/**
 	 * Call this to make sure everything is set up.
@@ -52,4 +62,9 @@ public:
 #endif // WITH_EDITOR
 	//~ End UObject interface
 
+private:
+	/** Stores the player. */
+	UPROPERTY()
+	TObjectPtr<UMediaPlayer> CurrentMediaPlayer = nullptr;
+
 };
diff --git a/Engine/Plugins/Media/MediaSourceManager/Source/MediaSourceManagerEditor/MediaSourceManagerEditor.Build.cs b/Engine/Plugins/Media/MediaSourceManager/Source/MediaSourceManagerEditor/MediaSourceManagerEditor.Build.cs
--- a/Engine/Plugins/Media/MediaSourceManager/Source/MediaSourceManagerEditor/MediaSourceManagerEditor.Build.cs
+++ b/Engine/Plugins/Media/MediaSourceManager/Source/MediaSourceManagerEditor/MediaSourceManagerEditor.Build.cs
@@ -15,6 +15,7 @@ public class MediaSourceManagerEditor : ModuleRules
 				"CoreUObject",
 				"InputCore",
 				"MediaAssets",
+				"MediaPlayerEditor",
 				"MediaSourceManager",
 				"Slate",
 				"SlateCore",
diff --git a/Engine/Plugins/Media/MediaSourceManager/Source/MediaSourceManagerEditor/Private/AssetTools/MediaSourceManagerActions.cpp b/Engine/Plugins/Media/MediaSourceManager/Source/MediaSourceManagerEditor/Private/AssetTools/MediaSourceManagerActions.cpp
--- a/Engine/Plugins/Media/MediaSourceManager/Source/MediaSourceManagerEditor/Private/AssetTools/MediaSourceManagerActions.cpp
+++ b/Engine/Plugins/Media/MediaSourceManager/Source/MediaSourceManagerEditor/Private/AssetTools/MediaSourceManagerActions.cpp
@@ -8,6 +8,17 @@
 
 #define LOCTEXT_NAMESPACE "MediaSourceManagerActions"
 
+/* FMediaSourceManagerActions constructors
+ *****************************************************************************/
+
+FMediaSourceManagerActions::FMediaSourceManagerActions(const TSharedRef<ISlateStyle>& InStyle)
+	: Style(InStyle)
+{
+}
+
+/* FAssetTypeActions_Base interface
+ *****************************************************************************/
+
 bool FMediaSourceManagerActions::CanFilter()
 {
 	return true;
@@ -50,7 +61,7 @@ void FMediaSourceManagerActions::OpenAssetEditor(const TArray<UObject*>& InObjec
 
 		if (MediaSourceManager != nullptr)
 		{
-			TSharedRef<FMediaSourceManagerEditorToolkit> EditorToolkit = MakeShareable(new FMediaSourceManagerEditorToolkit());
+			TSharedRef<FMediaSourceManagerEditorToolkit> EditorToolkit = MakeShareable(new FMediaSourceManagerEditorToolkit(Style));
 			EditorToolkit->Initialize(MediaSourceManager, Mode, EditWithinLevelEditor);
 		}
 	}
diff --git a/Engine/Plugins/Media/MediaSourceManager/Source/MediaSourceManagerEditor/Private/AssetTools/MediaSourceManagerActions.h b/Engine/Plugins/Media/MediaSourceManager/Source/MediaSourceManagerEditor/Private/AssetTools/MediaSourceManagerActions.h
--- a/Engine/Plugins/Media/MediaSourceManager/Source/MediaSourceManagerEditor/Private/AssetTools/MediaSourceManagerActions.h
+++ b/Engine/Plugins/Media/MediaSourceManager/Source/MediaSourceManagerEditor/Private/AssetTools/MediaSourceManagerActions.h
@@ -5,6 +5,8 @@
 #include "AssetTypeActions_Base.h"
 #include "Templates/SharedPointer.h"
 
+class ISlateStyle;
+
 /**
  * Implements an action for MediaSourceManager assets.
  */
@@ -12,6 +14,13 @@ class FMediaSourceManagerActions : public FAssetTypeActions_Base
 {
 public:
 
+	/**
+	 * Creates and initializes a new instance.
+	 *
+	 * @param InStyle The style set to use for asset editor toolkits.
+	 */
+	FMediaSourceManagerActions(const TSharedRef<ISlateStyle>& InStyle);
+
 	//~ FAssetTypeActions_Base interface
 	virtual bool CanFilter() override;
 	virtual FText GetAssetDescription(const FAssetData& AssetData) const override;
@@ -20,4 +29,9 @@ public:
 	virtual UClass* GetSupportedClass() const override;
 	virtual FColor GetTypeColor() const override;
 	virtual void OpenAssetEditor(const TArray<UObject*>& InObjects, TSharedPtr<class IToolkitHost> EditWithinLevelEditor = TSharedPtr<IToolkitHost>()) override;
+
+private:
+
+	/** Pointer to the style set to use for toolkits. */
+	TSharedRef<ISlateStyle> Style;
 };
diff --git a/Engine/Plugins/Media/MediaSourceManager/Source/MediaSourceManagerEditor/Private/MediaSourceManagerEditorModule.cpp b/Engine/Plugins/Media/MediaSourceManager/Source/MediaSourceManagerEditor/Private/MediaSourceManagerEditorModule.cpp
--- a/Engine/Plugins/Media/MediaSourceManager/Source/MediaSourceManagerEditor/Private/MediaSourceManagerEditorModule.cpp
+++ b/Engine/Plugins/Media/MediaSourceManager/Source/MediaSourceManagerEditor/Private/MediaSourceManagerEditorModule.cpp
@@ -4,6 +4,7 @@
 
 #include "AssetToolsModule.h"
 #include "AssetTools/MediaSourceManagerActions.h"
+#include "MediaSourceManagerEditorStyle.h"
 #include "Modules/ModuleManager.h"
 
 #define LOCTEXT_NAMESPACE "MediaSourceManagerEditorModule"
@@ -18,12 +19,14 @@ class FMediaSourceManagerEditorModule
 {
 public:
 	//~ IMediaSourceManagerEditorModule interface
-	
+	TSharedPtr<ISlateStyle> GetStyle() { return Style; }
 
 	//~ IModuleInterface interface
 
 	virtual void StartupModule() override
 	{
+		Style = MakeShareable(new FMediaSourceManagerEditorStyle());
+
 		RegisterAssetTools();
 	}
 
@@ -41,7 +44,7 @@ protected:
 	{
 		IAssetTools& AssetTools = FModuleManager::LoadModuleChecked<FAssetToolsModule>("AssetTools").Get();
 
-		RegisterAssetTypeAction(AssetTools, MakeShareable(new FMediaSourceManagerActions()));
+		RegisterAssetTypeAction(AssetTools, MakeShareable(new FMediaSourceManagerActions(Style.ToSharedRef())));
 	}
 
 	/**
@@ -73,6 +76,8 @@ protected:
 		}
 	}
 
+	/** Holds the plug-ins style set. */
+	TSharedPtr<ISlateStyle> Style;
 	/** The collection of registered asset type actions. */
 	TArray<TSharedRef<IAssetTypeActions>> RegisteredAssetTypeActions;
 };
diff --git a/Engine/Plugins/Media/MediaSourceManager/Source/MediaSourceManagerEditor/Private/MediaSourceManagerEditorStyle.cpp b/Engine/Plugins/Media/MediaSourceManager/Source/MediaSourceManagerEditor/Private/MediaSourceManagerEditorStyle.cpp
new file mode 100644
--- /dev/null
+++ b/Engine/Plugins/Media/MediaSourceManager/Source/MediaSourceManagerEditor/Private/MediaSourceManagerEditorStyle.cpp
@@ -0,0 +1,94 @@
+// Copyright Epic Games, Inc. All Rights Reserved.
+
+#include "MediaSourceManagerEditorStyle.h"
+
+#include "Brushes/SlateBorderBrush.h"
+#include "Brushes/SlateBoxBrush.h"
+#include "Brushes/SlateImageBrush.h"
+#include "Fonts/SlateFontInfo.h"
+#include "Math/Vector2D.h"
+#include "Misc/Paths.h"
+#include "Styling/SlateStyleRegistry.h"
+#include "Styling/SlateTypes.h"
+#include "Styling/CoreStyle.h"
+
+#define IMAGE_BRUSH(RelativePath, ...) FSlateImageBrush(RootToContentDir(RelativePath, TEXT(".png")), __VA_ARGS__)
+#define BOX_BRUSH(RelativePath, ...) FSlateBoxBrush(RootToContentDir(RelativePath, TEXT(".png")), __VA_ARGS__)
+#define BORDER_BRUSH(RelativePath, ...) FSlateBorderBrush(RootToContentDir(RelativePath, TEXT(".png")), __VA_ARGS__)
+#define DEFAULT_FONT(...) FCoreStyle::GetDefaultFontStyle(__VA_ARGS__)
+
+/* FMediaSourceManagerEditorStyle structors
+ *****************************************************************************/
+
+FMediaSourceManagerEditorStyle::FMediaSourceManagerEditorStyle()
+	: FSlateStyleSet("MediaSourceManagerEditorStyle")
+{
+	const FVector2D Icon8x8(8.0f, 8.0f);
+	const FVector2D Icon12x12(12.0f, 12.0f);
+	const FVector2D Icon16x16(16.0f, 16.0f);
+	const FVector2D Icon20x20(20.0f, 20.0f);
+	const FVector2D Icon40x40(40.0f, 40.0f);
+
+	SetContentRoot(FPaths::EnginePluginsDir() / TEXT("Media/MediaPlayerEditor/Content"));
+
+	// buttons
+	Set("MediaPlayerEditor.SourceButton", new IMAGE_BRUSH("btn_source_12x", Icon12x12));
+	Set("MediaPlayerEditor.GoButton", new IMAGE_BRUSH("btn_go_12x", Icon12x12));
+	Set("MediaPlayerEditor.ReloadButton", new IMAGE_BRUSH("btn_reload_12x", Icon12x12));
+	Set("MediaPlayerEditor.SettingsButton", new IMAGE_BRUSH("btn_settings_16x", Icon12x12));
+
+	// misc
+	Set("MediaPlateEditor.DragDropBorder", new BOX_BRUSH("border_dragdrop", 0.5f));
+	Set("MediaPlateEditor.MediaSourceOpened", new IMAGE_BRUSH("mediasource_opened", Icon8x8));
+
+	// tabs
+	Set("MediaPlateEditor.Tabs.Info", new IMAGE_BRUSH("tab_info_16x", Icon16x16));
+	Set("MediaPlateEditor.Tabs.Media", new IMAGE_BRUSH("tab_media_16x", Icon16x16));
+	Set("MediaPlateEditor.Tabs.Player", new IMAGE_BRUSH("tab_player_16x", Icon16x16));
+	Set("MediaPlateEditor.Tabs.Playlist", new IMAGE_BRUSH("tab_playlist_16x", Icon16x16));
+	Set("MediaPlateEditor.Tabs.Stats", new IMAGE_BRUSH("tab_stats_16x", Icon16x16));
+
+	// toolbar icons
+	Set("MediaPlateEditor.CloseMedia", new IMAGE_BRUSH("icon_eject_40x", Icon40x40));
+	Set("MediaPlateEditor.ForwardMedia", new IMAGE_BRUSH("icon_forward_40x", Icon40x40));
+	Set("MediaPlateEditor.ForwardMedia.Small", new IMAGE_BRUSH("icon_forward_40x", Icon20x20));
+	Set("MediaPlateEditor.NextMedia", new IMAGE_BRUSH("icon_step_40x", Icon40x40));
+	Set("MediaPlateEditor.NextMedia.Small", new IMAGE_BRUSH("icon_step_40x", Icon20x20));
+	Set("MediaPlateEditor.OpenMedia", new IMAGE_BRUSH("tab_media_16x", Icon40x40));
+	Set("MediaPlateEditor.PauseMedia", new IMAGE_BRUSH("icon_pause_40x", Icon40x40));
+	Set("MediaPlateEditor.PauseMedia.Small", new IMAGE_BRUSH("icon_pause_40x", Icon20x20));
+	Set("MediaPlateEditor.PlayMedia", new IMAGE_BRUSH("icon_play_40x", Icon40x40));
+	Set("MediaPlateEditor.PlayMedia.Small", new IMAGE_BRUSH("icon_play_40x", Icon20x20));
+	Set("MediaPlateEditor.PreviousMedia", new IMAGE_BRUSH("icon_step_back_40x", Icon40x40));
+	Set("MediaPlateEditor.PreviousMedia.Small", new IMAGE_BRUSH("icon_step_back_40x", Icon20x20));
+	Set("MediaPlateEditor.ReverseMedia", new IMAGE_BRUSH("icon_reverse_40x", Icon40x40));
+	Set("MediaPlateEditor.ReverseMedia.Small", new IMAGE_BRUSH("icon_reverse_40x", Icon20x20));
+	Set("MediaPlateEditor.RewindMedia", new IMAGE_BRUSH("icon_rewind_40x", Icon40x40));
+	Set("MediaPlateEditor.RewindMedia.Small", new IMAGE_BRUSH("icon_rewind_40x", Icon20x20));
+	Set("MediaPlateEditor.StopMedia", new IMAGE_BRUSH("icon_stop_40x", Icon40x40));
+	Set("MediaPlateEditor.StopMedia.Small", new IMAGE_BRUSH("icon_stop_40x", Icon20x20));
+
+	// scrubber
+	Set("MediaPlayerEditor.Scrubber", FSliderStyle()
+		.SetNormalBarImage(FSlateColorBrush(FColor::White))
+		.SetDisabledBarImage(FSlateColorBrush(FLinearColor::Gray))
+		.SetNormalThumbImage(IMAGE_BRUSH("scrubber", FVector2D(2.0f, 10.0f)))
+		.SetHoveredThumbImage(IMAGE_BRUSH("scrubber", FVector2D(2.0f, 10.0f)))
+		.SetDisabledThumbImage(IMAGE_BRUSH("scrubber", FVector2D(2.0f, 10.0f)))
+		.SetBarThickness(2.0f)
+	);
+
+	Set("MediaPlateEditor.ViewportFont", DEFAULT_FONT("Regular", 18));
+
+	FSlateStyleRegistry::RegisterSlateStyle(*this);
+}
+
+FMediaSourceManagerEditorStyle::~FMediaSourceManagerEditorStyle()
+{
+	FSlateStyleRegistry::UnRegisterSlateStyle(*this);
+}
+
+#undef IMAGE_BRUSH
+#undef BOX_BRUSH
+#undef BORDER_BRUSH
+#undef DEFAULT_FONT
diff --git a/Engine/Plugins/Media/MediaSourceManager/Source/MediaSourceManagerEditor/Private/MediaSourceManagerEditorStyle.h b/Engine/Plugins/Media/MediaSourceManager/Source/MediaSourceManagerEditor/Private/MediaSourceManagerEditorStyle.h
new file mode 100644
--- /dev/null
+++ b/Engine/Plugins/Media/MediaSourceManager/Source/MediaSourceManagerEditor/Private/MediaSourceManagerEditorStyle.h
@@ -0,0 +1,18 @@
+// Copyright Epic Games, Inc. All Rights Reserved.
+
+#pragma once
+
+#include "Styling/SlateStyle.h"
+
+/**
+ * Implements the visual style of the media source manager editor UI.
+ */
+class FMediaSourceManagerEditorStyle : public FSlateStyleSet
+{
+public:
+	/** Default constructor. */
+	FMediaSourceManagerEditorStyle();
+
+	 /** Destructor. */
+	~FMediaSourceManagerEditorStyle();
+};
diff --git a/Engine/Plugins/Media/MediaSourceManager/Source/MediaSourceManagerEditor/Private/Toolkits/MediaSourceManagerEditorToolkit.cpp b/Engine/Plugins/Media/MediaSourceManager/Source/MediaSourceManagerEditor/Private/Toolkits/MediaSourceManagerEditorToolkit.cpp
--- a/Engine/Plugins/Media/MediaSourceManager/Source/MediaSourceManagerEditor/Private/Toolkits/MediaSourceManagerEditorToolkit.cpp
+++ b/Engine/Plugins/Media/MediaSourceManager/Source/MediaSourceManagerEditor/Private/Toolkits/MediaSourceManagerEditorToolkit.cpp
@@ -5,9 +5,11 @@
 #include "Editor.h"
 #include "EditorReimportHandler.h"
 #include "MediaSourceManager.h"
+#include "MediaTexture.h"
 #include "SlateOptMacros.h"
 #include "Styling/AppStyle.h"
 #include "Widgets/Docking/SDockTab.h"
+#include "Widgets/SMediaPlayerEditorViewer.h"
 #include "Widgets/SMediaSourceManagerSources.h"
 
 #define LOCTEXT_NAMESPACE "FMediaSourceManagerEditorToolkit"
@@ -16,10 +18,12 @@ namespace MediaSourceManagerEditorToolkit
 {
 	static const FName AppIdentifier("MediaSourceManagerEditorApp");
 	static const FName SourcesTabId("Sources");
+	static const FName ViewerTabId("Viewer");
 }
 
-FMediaSourceManagerEditorToolkit::FMediaSourceManagerEditorToolkit()
+FMediaSourceManagerEditorToolkit::FMediaSourceManagerEditorToolkit(const TSharedRef<ISlateStyle>& InStyle)
 	: MediaSourceManager(nullptr)
+	, Style(InStyle)
 {
 }
 
@@ -46,18 +50,50 @@ void FMediaSourceManagerEditorToolkit::Initialize(UMediaSourceManager* InMediaSo
 	GEditor->RegisterForUndo(this);
 
 	// create tab layout
-	const TSharedRef<FTabManager::FLayout> Layout = FTabManager::NewLayout("Standalone_MediaSourceManagerEditor_v0.1b")
+	TSharedRef<FTabManager::FArea> PrimaryArea = FTabManager::NewPrimaryArea()
+		->SetOrientation(Orient_Horizontal);
+	TSharedRef<FTabManager::FSplitter> PreviewSplitter = FTabManager::NewSplitter();
+	PreviewSplitter->SetOrientation(Orient_Vertical);
+
+	// Base tab name for each preview window.
+	FString TabId = MediaSourceManagerEditorToolkit::ViewerTabId.ToString();
+	int32 OriginalLength = TabId.Len();
+
+	// Add each channel.
+	PreviewTabIds.SetNum(MediaSourceManager->Channels.Num());
+	for (int Index = 0; Index < MediaSourceManager->Channels.Num(); ++Index)
+	{
+		// Construct tab name,
+		TabId.LeftInline(OriginalLength, false);
+		TabId.AppendInt(Index);
+		PreviewTabIds[Index] = FName(*TabId);
+
+		// Add to splitter.
+		PreviewSplitter->Split
+		(
+			FTabManager::NewStack()
+			->AddTab(FName(*TabId), ETabState::OpenedTab)
+		);
+	}
+
+	PrimaryArea->Split
+	(
+		// Sources tab.
+		FTabManager::NewStack()
+		->AddTab(MediaSourceManagerEditorToolkit::SourcesTabId, ETabState::OpenedTab)
+		->SetHideTabWell(true)
+		->SetSizeCoefficient(0.4f)
+	)
+	->Split
+	(
+		// Channel previews.
+		PreviewSplitter
+	);
+
+	const TSharedRef<FTabManager::FLayout> Layout = FTabManager::NewLayout("Standalone_MediaSourceManagerEditor_v0.1d4")
 		->AddArea
 		(
-			FTabManager::NewPrimaryArea()
-				->SetOrientation(Orient_Vertical)
-				->Split
-				(
-					// Sources tab.
-					FTabManager::NewStack()
-						->AddTab(MediaSourceManagerEditorToolkit::SourcesTabId, ETabState::OpenedTab)
-						->SetHideTabWell(true)
-				)
+			PrimaryArea
 		);
 
 	FAssetEditorToolkit::InitAssetEditor(
@@ -93,12 +129,28 @@ void FMediaSourceManagerEditorToolkit::RegisterTabSpawners(const TSharedRef<FTab
 		.SetDisplayName(LOCTEXT("SourcesTabName", "Sources"))
 		.SetGroup(WorkspaceMenuCategoryRef)
 		.SetIcon(FSlateIcon(FAppStyle::GetAppStyleSetName(), "LevelEditor.Tabs.Details"));
+
+	// Viewer tab.
+	for (int Index = 0; Index < PreviewTabIds.Num(); Index++)
+	{
+		InTabManager->RegisterTabSpawner(PreviewTabIds[Index],
+			FOnSpawnTab::CreateSP(this, &FMediaSourceManagerEditorToolkit::HandlePreviewTabManagerSpawnTab,
+				Index))
+			.SetDisplayName(FText::Format(LOCTEXT("ChannelTabName", "Channel {0}"),
+				FText::AsNumber(Index)))
+			.SetGroup(WorkspaceMenuCategoryRef)
+			.SetIcon(FSlateIcon(FAppStyle::GetAppStyleSetName(), "MediaPlateEditor.Tabs.Player"));
+	}
 }
 
 void FMediaSourceManagerEditorToolkit::UnregisterTabSpawners(const TSharedRef<class FTabManager>& InTabManager)
 {
 	FAssetEditorToolkit::UnregisterTabSpawners(InTabManager);
 
+	for (int Index = 0; Index < PreviewTabIds.Num(); Index++)
+	{
+		InTabManager->UnregisterTabSpawner(PreviewTabIds[Index]);
+	}
 	InTabManager->UnregisterTabSpawner(MediaSourceManagerEditorToolkit::SourcesTabId);
 }
 
@@ -152,4 +204,31 @@ TSharedRef<SDockTab> FMediaSourceManagerEditorToolkit::HandleTabManagerSpawnTab(
 		];
 }
 
+
+TSharedRef<SDockTab> FMediaSourceManagerEditorToolkit::HandlePreviewTabManagerSpawnTab(
+	const FSpawnTabArgs& Args, int32 ChannelIndex)
+{
+	TSharedPtr<SWidget> TabWidget = SNullWidget::NullWidget;
+
+	if (MediaSourceManager != nullptr)
+	{
+		if (ChannelIndex < MediaSourceManager->Channels.Num())
+		{
+			UMediaSourceManagerChannel* Channel = MediaSourceManager->Channels[ChannelIndex];
+			if (Channel != nullptr)
+			{
+				UMediaPlayer* MediaPlayer = Channel->GetMediaPlayer();
+				UMediaTexture* MediaTexture = Cast<UMediaTexture>(Channel->OutTexture);
+				TabWidget = SNew(SMediaPlayerEditorViewer, *MediaPlayer, MediaTexture, Style, false);
+			}
+		}
+	}
+	
+	return SNew(SDockTab)
+		.TabRole(ETabRole::PanelTab)
+		[
+			TabWidget.ToSharedRef()
+		];
+}
+
 #undef LOCTEXT_NAMESPACE
diff --git a/Engine/Plugins/Media/MediaSourceManager/Source/MediaSourceManagerEditor/Private/Toolkits/MediaSourceManagerEditorToolkit.h b/Engine/Plugins/Media/MediaSourceManager/Source/MediaSourceManagerEditor/Private/Toolkits/MediaSourceManagerEditorToolkit.h
--- a/Engine/Plugins/Media/MediaSourceManager/Source/MediaSourceManagerEditor/Private/Toolkits/MediaSourceManagerEditorToolkit.h
+++ b/Engine/Plugins/Media/MediaSourceManager/Source/MediaSourceManagerEditor/Private/Toolkits/MediaSourceManagerEditorToolkit.h
@@ -7,6 +7,7 @@
 #include "Toolkits/AssetEditorToolkit.h"
 #include "UObject/GCObject.h"
 
+class ISlateStyle;
 class UMediaSourceManager;
 
 /**
@@ -20,8 +21,10 @@ class FMediaSourceManagerEditorToolkit
 public:
 	/**
 	 * Creates and initializes a new instance.
+	 * 
+	 * @param InStyle The style set to use.
 	 */
-	FMediaSourceManagerEditorToolkit();
+	FMediaSourceManagerEditorToolkit(const TSharedRef<ISlateStyle>& InStyle);
 
 	/** Virtual destructor. */
 	virtual ~FMediaSourceManagerEditorToolkit();
@@ -64,7 +67,15 @@ private:
 
 	/** Callback for spawning tabs. */
 	TSharedRef<SDockTab> HandleTabManagerSpawnTab(const FSpawnTabArgs& Args, FName TabIdentifier);
+	/** Callback for spawning tabs. */
+	TSharedRef<SDockTab> HandlePreviewTabManagerSpawnTab(const FSpawnTabArgs& Args, int32 ChannelIndex);
 
 	/** The media source manager asset being edited. */
 	TObjectPtr<UMediaSourceManager> MediaSourceManager;
+
+	/** Pointer to the style set to use for toolkits. */
+	TSharedRef<ISlateStyle> Style;
+
+	/** Tab Ids for channel preview windows. */
+	TArray<FName> PreviewTabIds;
 };
diff --git a/Engine/Plugins/Media/MediaSourceManager/Source/MediaSourceManagerEditor/Private/Widgets/SMediaSourceManagerChannel.cpp b/Engine/Plugins/Media/MediaSourceManager/Source/MediaSourceManagerEditor/Private/Widgets/SMediaSourceManagerChannel.cpp
--- a/Engine/Plugins/Media/MediaSourceManager/Source/MediaSourceManagerEditor/Private/Widgets/SMediaSourceManagerChannel.cpp
+++ b/Engine/Plugins/Media/MediaSourceManager/Source/MediaSourceManagerEditor/Private/Widgets/SMediaSourceManagerChannel.cpp
@@ -51,6 +51,12 @@ void SMediaSourceManagerChannel::Construct(const FArguments& InArgs,
 		];
 
 	Refresh();
+
+	// Start playing.
+	if (ChannelPtr != nullptr)
+	{
+		ChannelPtr->Play();
+	}
 }
 
 void SMediaSourceManagerChannel::OnDragEnter(const FGeometry& MyGeometry, const FDragDropEvent& DragDropEvent)
@@ -90,6 +96,8 @@ FReply SMediaSourceManagerChannel::OnDrop(const FGeometry& MyGeometry, const FDr
 					UMediaSourceManagerInputMediaSource* Input = NewObject<UMediaSourceManagerInputMediaSource>(Channel);
 					Input->MediaSource = MediaSource;
 					Channel->Input = Input;
+
+					Channel->Play();
 						
 					Refresh();
 					break;
diff --git a/Engine/Plugins/Media/MediaSourceManager/Source/MediaSourceManagerEditor/Public/MediaSourceManagerEditorModule.h b/Engine/Plugins/Media/MediaSourceManager/Source/MediaSourceManagerEditor/Public/MediaSourceManagerEditorModule.h
--- a/Engine/Plugins/Media/MediaSourceManager/Source/MediaSourceManagerEditor/Public/MediaSourceManagerEditorModule.h
+++ b/Engine/Plugins/Media/MediaSourceManager/Source/MediaSourceManagerEditor/Public/MediaSourceManagerEditorModule.h
@@ -4,10 +4,13 @@
 
 #include "Logging/LogMacros.h"
 #include "Modules/ModuleInterface.h"
+#include "Templates/SharedPointer.h"
 
 /** Log category for this module. */
 DECLARE_LOG_CATEGORY_EXTERN(LogMediaSourceManagerEditor, Log, All);
 
+class ISlateStyle;
+
 /**
 * Interface for the MediaSourceManagerEditor module.
 */
@@ -16,4 +19,9 @@ class MEDIASOURCEMANAGEREDITOR_API IMediaSourceManagerEditorModule
 {
 public:
 
+	/**
+	 * Get the style used by this module.
+	 **/
+	virtual TSharedPtr<ISlateStyle> GetStyle() = 0;
+
 };
