MediaSourceManager: Added editor tabs for the media source manager.

#rnx

diff --git a/Engine/Plugins/Media/MediaSourceManager/Source/MediaSourceManager/Private/MediaSourceManagerModule.cpp b/Engine/Plugins/Media/MediaSourceManager/Source/MediaSourceManager/Private/MediaSourceManagerModule.cpp
--- a/Engine/Plugins/Media/MediaSourceManager/Source/MediaSourceManager/Private/MediaSourceManagerModule.cpp
+++ b/Engine/Plugins/Media/MediaSourceManager/Source/MediaSourceManager/Private/MediaSourceManagerModule.cpp
@@ -1,8 +1,15 @@
 // Copyright Epic Games, Inc. All Rights Reserved.
 
 #include "MediaSourceManagerModule.h"
+
+#include "MediaSourceManagerSettings.h"
 #include "Modules/ModuleManager.h"
 
+#if WITH_EDITOR
+#include "ISettingsModule.h"
+#include "ISettingsSection.h"
+#endif //WITH_EDITOR
+
 #define LOCTEXT_NAMESPACE "MediaSourceManagerModule"
 
 DEFINE_LOG_CATEGORY(LogMediaSourceManager);
@@ -21,11 +28,50 @@ public:
 
 	virtual void StartupModule() override
 	{
+		RegisterSettings();
 	}
 
 	virtual void ShutdownModule() override
 	{
+		UnregisterSettings();
+	}
+
+private:
+
+	void RegisterSettings()
+	{
+#if WITH_EDITOR
+		if (GIsEditor)
+		{
+			// register settings
+			ISettingsModule* SettingsModule = FModuleManager::GetModulePtr<ISettingsModule>("Settings");
+			if (SettingsModule != nullptr)
+			{
+				SettingsModule->RegisterSettings("Project", "Plugins", "MediaSourceManager",
+					LOCTEXT("SettingsName", "Media Source Manager"),
+					LOCTEXT("Description", "Configure the Media Source Manager."),
+					GetMutableDefault<UMediaSourceManagerSettings>()
+				);
+			}
+		}
+#endif //WITH_EDITOR
 	}
+
+	void UnregisterSettings()
+	{
+#if WITH_EDITOR
+		if (GIsEditor)
+		{
+			// unregister settings
+			ISettingsModule* SettingsModule = FModuleManager::GetModulePtr<ISettingsModule>("Settings");
+			if (SettingsModule != nullptr)
+			{
+				SettingsModule->UnregisterSettings("Project", "Plugins", "MediaSourceManager");
+			}
+		}
+#endif //WITH_EDITOR
+	}
+
 };
 
 IMPLEMENT_MODULE(FMediaSourceManagerModule, MediaSourceManager);
diff --git a/Engine/Plugins/Media/MediaSourceManager/Source/MediaSourceManager/Private/MediaSourceManagerSettings.cpp b/Engine/Plugins/Media/MediaSourceManager/Source/MediaSourceManager/Private/MediaSourceManagerSettings.cpp
new file mode 100644
--- /dev/null
+++ b/Engine/Plugins/Media/MediaSourceManager/Source/MediaSourceManager/Private/MediaSourceManagerSettings.cpp
@@ -0,0 +1,21 @@
+// Copyright Epic Games, Inc. All Rights Reserved.
+
+#include "MediaSourceManagerSettings.h"
+
+#include "MediaSourceManager.h"
+
+UMediaSourceManager* UMediaSourceManagerSettings::GetManager() const
+{
+	return Manager.LoadSynchronous();
+}
+
+#if WITH_EDITOR
+
+void UMediaSourceManagerSettings::SetManager(UMediaSourceManager* InManager)
+{
+	Manager = InManager;
+	SaveConfig();
+	OnManagerChanged.Broadcast();
+}
+
+#endif // WITH_EDITOR
diff --git a/Engine/Plugins/Media/MediaSourceManager/Source/MediaSourceManager/Public/MediaSourceManagerSettings.h b/Engine/Plugins/Media/MediaSourceManager/Source/MediaSourceManager/Public/MediaSourceManagerSettings.h
new file mode 100644
--- /dev/null
+++ b/Engine/Plugins/Media/MediaSourceManager/Source/MediaSourceManager/Public/MediaSourceManagerSettings.h
@@ -0,0 +1,45 @@
+// Copyright Epic Games, Inc. All Rights Reserved.
+
+#pragma once
+
+#include "CoreTypes.h"
+#include "UObject/Object.h"
+#include "UObject/ObjectMacros.h"
+#include "UObject/SoftObjectPtr.h"
+
+#include "MediaSourceManagerSettings.generated.h"
+
+class UMediaSourceManager;
+
+/**
+* Settings for the media source manager.
+*/
+UCLASS(config=Game, defaultconfig)
+class MEDIASOURCEMANAGER_API UMediaSourceManagerSettings : public UObject
+{
+	GENERATED_BODY()
+
+public:
+	/**
+	 * Call this to get the manager.
+	 */
+	UMediaSourceManager* GetManager() const;
+
+#if WITH_EDITOR
+
+	/** Register with this to get a callback when the manager changes. */
+	FSimpleMulticastDelegate OnManagerChanged;
+
+	/**
+	 * Call this to set the manager.
+	 */
+	void SetManager(UMediaSourceManager* InManager);
+
+#endif // WITH_EDITOR
+
+private:
+
+	/** Stores the current manager asset. */
+	UPROPERTY(config, EditAnywhere, Category = "MediaSourceManager")
+	TSoftObjectPtr<UMediaSourceManager> Manager;
+};
diff --git a/Engine/Plugins/Media/MediaSourceManager/Source/MediaSourceManagerEditor/MediaSourceManagerEditor.Build.cs b/Engine/Plugins/Media/MediaSourceManager/Source/MediaSourceManagerEditor/MediaSourceManagerEditor.Build.cs
--- a/Engine/Plugins/Media/MediaSourceManager/Source/MediaSourceManagerEditor/MediaSourceManagerEditor.Build.cs
+++ b/Engine/Plugins/Media/MediaSourceManager/Source/MediaSourceManagerEditor/MediaSourceManagerEditor.Build.cs
@@ -17,9 +17,11 @@ public class MediaSourceManagerEditor : ModuleRules
 				"MediaAssets",
 				"MediaPlayerEditor",
 				"MediaSourceManager",
+				"PropertyEditor",
 				"Slate",
 				"SlateCore",
 				"UnrealEd",
+				"WorkspaceMenuStructure",
 			}
 			);
 	}
diff --git a/Engine/Plugins/Media/MediaSourceManager/Source/MediaSourceManagerEditor/Private/MediaSourceManagerEditorModule.cpp b/Engine/Plugins/Media/MediaSourceManager/Source/MediaSourceManagerEditor/Private/MediaSourceManagerEditorModule.cpp
--- a/Engine/Plugins/Media/MediaSourceManager/Source/MediaSourceManagerEditor/Private/MediaSourceManagerEditorModule.cpp
+++ b/Engine/Plugins/Media/MediaSourceManager/Source/MediaSourceManagerEditor/Private/MediaSourceManagerEditorModule.cpp
@@ -6,6 +6,12 @@
 #include "AssetTools/MediaSourceManagerActions.h"
 #include "MediaSourceManagerEditorStyle.h"
 #include "Modules/ModuleManager.h"
+#include "UObject/NameTypes.h"
+#include "Widgets/Docking/SDockTab.h"
+#include "Widgets/SMediaSourceManager.h"
+#include "Widgets/SMediaSourceManagerPreviews.h"
+#include "WorkspaceMenuStructure.h"
+#include "WorkspaceMenuStructureModule.h"
 
 #define LOCTEXT_NAMESPACE "MediaSourceManagerEditorModule"
 
@@ -28,10 +34,12 @@ public:
 		Style = MakeShareable(new FMediaSourceManagerEditorStyle());
 
 		RegisterAssetTools();
+		RegisterTabSpawners();
 	}
 
 	virtual void ShutdownModule() override
 	{
+		UnregisterTabSpawners();
 		UnregisterAssetTools();
 	}
 
@@ -76,12 +84,87 @@ protected:
 		}
 	}
 
+	/**
+	 * Register functions to spawn tabs.
+	 */
+	void RegisterTabSpawners()
+	{
+		// Add Media group.
+		const IWorkspaceMenuStructure& MenuStructure = WorkspaceMenu::GetMenuStructure();
+		TSharedRef<FWorkspaceItem> MediaGroup = MenuStructure.GetLevelEditorCategory()->AddGroup(
+			LOCTEXT("WorkspaceMenu_MediaCategory", "Media"),
+			FSlateIcon(FAppStyle::GetAppStyleSetName(), "SequenceRecorder.TabIcon"),
+			true);
+		TSharedRef<FWorkspaceItem> MediaSourceManagerGroup = MediaGroup->AddGroup(
+			LOCTEXT("WorkspaceMenu_MediaSourceManagerCategory", "Media Source Manager"),
+			FSlateIcon(FAppStyle::GetAppStyleSetName(), "SequenceRecorder.TabIcon"),
+			true);
+
+		// Add media source manager tab.
+		FGlobalTabmanager::Get()->RegisterNomadTabSpawner(ManagerTabName,
+			FOnSpawnTab::CreateStatic(&FMediaSourceManagerEditorModule::SpawnMediaSourceManagerTab))
+			.SetGroup(MediaSourceManagerGroup)
+			.SetDisplayName(LOCTEXT("ManagerTabTitle", "Media Source Manager"))
+			.SetTooltipText(LOCTEXT("ManagerTooltipText", "Open the Media Source Manager tab."))
+			.SetIcon(FSlateIcon(FAppStyle::GetAppStyleSetName(), "SequenceRecorder.TabIcon"));
+
+		// Add preview tab.
+		FGlobalTabmanager::Get()->RegisterNomadTabSpawner(PreviewTabName,
+			FOnSpawnTab::CreateStatic(&FMediaSourceManagerEditorModule::SpawnMediaSourceManagerPreviewTab))
+			.SetGroup(MediaSourceManagerGroup)
+			.SetDisplayName(LOCTEXT("PreviewTabTitle", "Media Source Manager Preview"))
+			.SetTooltipText(LOCTEXT("PreviewTooltipText", "Open the Media Source Manager Preview tab."))
+			.SetIcon(FSlateIcon(FAppStyle::GetAppStyleSetName(), "LevelEditor.Tabs.Viewports"));
+	}
+
+	/**
+	 * Remove functions that spawned tabs.
+	 */
+	void UnregisterTabSpawners()
+	{
+		if (FSlateApplication::IsInitialized())
+		{
+			FGlobalTabmanager::Get()->UnregisterNomadTabSpawner(PreviewTabName);
+			FGlobalTabmanager::Get()->UnregisterNomadTabSpawner(ManagerTabName);
+		}
+	}
+
+	/**
+	 * Callback to spawn the manager tab.
+	 */
+	static TSharedRef<SDockTab> SpawnMediaSourceManagerTab(const FSpawnTabArgs& SpawnTabArgs)
+	{
+		return SNew(SDockTab)
+			.TabRole(ETabRole::NomadTab)
+			[
+				SNew(SMediaSourceManager)
+			];
+	}
+
+	/**
+	 * Callback to spawn the preview tab.
+	 */
+	static TSharedRef<SDockTab> SpawnMediaSourceManagerPreviewTab(const FSpawnTabArgs& SpawnTabArgs)
+	{
+		return SNew(SDockTab)
+			.TabRole(ETabRole::NomadTab)
+			[
+				SNew(SMediaSourceManagerPreviews)
+			];
+	}
+
 	/** Holds the plug-ins style set. */
 	TSharedPtr<ISlateStyle> Style;
 	/** The collection of registered asset type actions. */
 	TArray<TSharedRef<IAssetTypeActions>> RegisteredAssetTypeActions;
+	/** Names for tabs. */
+	static FLazyName ManagerTabName;
+	static FLazyName PreviewTabName;
 };
 
+FLazyName FMediaSourceManagerEditorModule::ManagerTabName(TEXT("MediaSourceManager"));
+FLazyName FMediaSourceManagerEditorModule::PreviewTabName(TEXT("MediaSourceManagerPreview"));
+
 IMPLEMENT_MODULE(FMediaSourceManagerEditorModule, MediaSourceManagerEditor);
 
 #undef LOCTEXT_NAMESPACE
diff --git a/Engine/Plugins/Media/MediaSourceManager/Source/MediaSourceManagerEditor/Private/Toolkits/MediaSourceManagerEditorToolkit.cpp b/Engine/Plugins/Media/MediaSourceManager/Source/MediaSourceManagerEditor/Private/Toolkits/MediaSourceManagerEditorToolkit.cpp
--- a/Engine/Plugins/Media/MediaSourceManager/Source/MediaSourceManagerEditor/Private/Toolkits/MediaSourceManagerEditorToolkit.cpp
+++ b/Engine/Plugins/Media/MediaSourceManager/Source/MediaSourceManagerEditor/Private/Toolkits/MediaSourceManagerEditorToolkit.cpp
@@ -194,7 +194,7 @@ TSharedRef<SDockTab> FMediaSourceManagerEditorToolkit::HandleTabManagerSpawnTab(
 
 	if (TabIdentifier == MediaSourceManagerEditorToolkit::SourcesTabId)
 	{
-		TabWidget = SNew(SMediaSourceManagerSources, *MediaSourceManager);
+		TabWidget = SNew(SMediaSourceManagerSources, MediaSourceManager);
 	}
 
 	return SNew(SDockTab)
diff --git a/Engine/Plugins/Media/MediaSourceManager/Source/MediaSourceManagerEditor/Private/Widgets/SMediaSourceManager.cpp b/Engine/Plugins/Media/MediaSourceManager/Source/MediaSourceManagerEditor/Private/Widgets/SMediaSourceManager.cpp
new file mode 100644
--- /dev/null
+++ b/Engine/Plugins/Media/MediaSourceManager/Source/MediaSourceManagerEditor/Private/Widgets/SMediaSourceManager.cpp
@@ -0,0 +1,100 @@
+// Copyright Epic Games, Inc. All Rights Reserved.
+
+#include "Widgets/SMediaSourceManager.h"
+
+#include "MediaSourceManager.h"
+#include "MediaSourceManagerEditorModule.h"
+#include "MediaSourceManagerSettings.h"
+#include "MediaTexture.h"
+#include "PropertyCustomizationHelpers.h"
+#include "Widgets/Input/SComboButton.h"
+#include "Widgets/Layout/SBox.h"
+#include "Widgets/SMediaPlayerEditorViewer.h"
+#include "Widgets/SMediaSourceManagerSources.h"
+#include "Widgets/Text/STextBlock.h"
+
+#define LOCTEXT_NAMESPACE "SMediaSourceManager"
+
+void SMediaSourceManager::Construct(const FArguments& InArgs)
+{
+	MediaSourceManagerPtr = GetDefault<UMediaSourceManagerSettings>()->GetManager();
+	
+	ChildSlot
+		[
+			SNew(SVerticalBox)
+
+			// Manager selector.
+			+ SVerticalBox::Slot()
+				.AutoHeight()
+				.Padding(2)
+				.VAlign(VAlign_Center)
+				[
+					SAssignNew(ManagerButton, SComboButton)
+						.OnGetMenuContent(this, &SMediaSourceManager::GetManagerPicker)
+						.ButtonContent()
+						[
+							SNew(STextBlock)
+								.Text(this, &SMediaSourceManager::GetManagerName)
+								.Font(FAppStyle::GetFontStyle(TEXT("PropertyWindow.NormalFont")))
+						]
+				]
+
+			// Sources.
+			+ SVerticalBox::Slot()
+				.Padding(2)
+				.VAlign(VAlign_Center)
+				[
+					SAssignNew(SourcesWidget, SMediaSourceManagerSources, MediaSourceManagerPtr.Get())
+				]
+		];
+}
+
+FText SMediaSourceManager::GetManagerName() const
+{
+	UMediaSourceManager* MediaSourceManager = MediaSourceManagerPtr.Get();
+
+	if (MediaSourceManager != nullptr)
+	{
+		return FText::FromName(MediaSourceManager->GetFName());
+	}
+	return LOCTEXT("SelectManager", "Select a MediaSourceManager asset.");
+}
+
+TSharedRef<SWidget> SMediaSourceManager::GetManagerPicker()
+{
+	UMediaSourceManager* MediaSourceManager = MediaSourceManagerPtr.Get();
+	FAssetData CurrentAssetData = MediaSourceManager ? FAssetData(MediaSourceManager) : FAssetData();
+
+	TArray<const UClass*> ClassFilters;
+	ClassFilters.Add(UMediaSourceManager::StaticClass());
+
+	return PropertyCustomizationHelpers::MakeAssetPickerWithMenu(
+		CurrentAssetData,
+		MediaSourceManagerPtr != nullptr,
+		false,
+		ClassFilters,
+		TArray<UFactory*>(),
+		FOnShouldFilterAsset::CreateLambda([CurrentAssetData](const FAssetData& InAssetData) { return InAssetData == CurrentAssetData; }),
+		FOnAssetSelected::CreateRaw(this, &SMediaSourceManager::NewManagerSelected),
+		FSimpleDelegate()
+	);
+}
+
+void SMediaSourceManager::NewManagerSelected(const FAssetData& AssetData)
+{
+	// Close the combo box.
+	ManagerButton->SetIsOpen(false);
+
+	// Set new manager.
+	UMediaSourceManager* Manager = Cast<UMediaSourceManager>(AssetData.GetAsset());
+	MediaSourceManagerPtr = Manager;
+	GetMutableDefault<UMediaSourceManagerSettings>()->SetManager(Manager);
+
+	// Refresh sources widget.
+	if (SourcesWidget != nullptr)
+	{
+		SourcesWidget->SetManager(Manager);
+	}
+}
+
+#undef LOCTEXT_NAMESPACE
diff --git a/Engine/Plugins/Media/MediaSourceManager/Source/MediaSourceManagerEditor/Private/Widgets/SMediaSourceManager.h b/Engine/Plugins/Media/MediaSourceManager/Source/MediaSourceManagerEditor/Private/Widgets/SMediaSourceManager.h
new file mode 100644
--- /dev/null
+++ b/Engine/Plugins/Media/MediaSourceManager/Source/MediaSourceManagerEditor/Private/Widgets/SMediaSourceManager.h
@@ -0,0 +1,57 @@
+// Copyright Epic Games, Inc. All Rights Reserved.
+
+#pragma once
+
+#include "Widgets/DeclarativeSyntaxSupport.h"
+#include "Widgets/SCompoundWidget.h"
+
+class SComboButton;
+class SDockTab;
+class SMediaSourceManagerSources;
+class STextBlock;
+class SVerticalBox;
+class UMediaSourceManager;
+
+/**
+ * Implements the MediaSourceManager widget.
+ */
+class SMediaSourceManager
+	: public SCompoundWidget
+{
+public:
+
+	SLATE_BEGIN_ARGS(SMediaSourceManager) { }
+	SLATE_END_ARGS()
+
+	/**
+	 * Construct this widget.
+	 *
+	 * @param InArgs					The declaration data for this widget.
+	 */
+	void Construct(const FArguments& InArgs);
+
+private:
+
+	/** Pointer to the object that is being viewed. */
+	TWeakObjectPtr<UMediaSourceManager> MediaSourceManagerPtr;
+	/** Widget that shows our manager. */
+	TSharedPtr<SComboButton> ManagerButton;
+	/** Widget that shows our sources. */
+	TSharedPtr<SMediaSourceManagerSources> SourcesWidget;
+
+	/**
+	 * Gets the name of our manager.
+	 */
+	FText GetManagerName() const;
+
+	/**
+	 * Get a widget to pick the manager asset.
+	 */
+	TSharedRef<SWidget> GetManagerPicker();
+
+	/**
+	 * Called when a new manager is picked.
+	 */
+	void NewManagerSelected(const FAssetData& AssetData);
+
+};
diff --git a/Engine/Plugins/Media/MediaSourceManager/Source/MediaSourceManagerEditor/Private/Widgets/SMediaSourceManagerPreviews.cpp b/Engine/Plugins/Media/MediaSourceManager/Source/MediaSourceManagerEditor/Private/Widgets/SMediaSourceManagerPreviews.cpp
new file mode 100644
--- /dev/null
+++ b/Engine/Plugins/Media/MediaSourceManager/Source/MediaSourceManagerEditor/Private/Widgets/SMediaSourceManagerPreviews.cpp
@@ -0,0 +1,83 @@
+// Copyright Epic Games, Inc. All Rights Reserved.
+
+#include "Widgets/SMediaSourceManagerPreviews.h"
+
+#include "MediaSourceManager.h"
+#include "MediaSourceManagerEditorModule.h"
+#include "MediaSourceManagerSettings.h"
+#include "MediaTexture.h"
+#include "Modules/ModuleManager.h"
+#include "Widgets/Layout/SScrollBox.h"
+#include "Widgets/SMediaPlayerEditorViewer.h"
+
+#define LOCTEXT_NAMESPACE "SMediaSourceManagerPreviews"
+
+void SMediaSourceManagerPreviews::Construct(const FArguments& InArg)
+{
+	ChildSlot
+		[
+			SAssignNew(ChannelsContainer, SHorizontalBox)
+		];
+
+	// Get style.
+	IMediaSourceManagerEditorModule* Module = FModuleManager::GetModulePtr
+		<IMediaSourceManagerEditorModule>("MediaSourceManagerEditor");
+	if (Module != nullptr)
+	{
+		Style = Module->GetStyle();
+	}
+
+	SetManager();
+	GetMutableDefault<UMediaSourceManagerSettings>()->OnManagerChanged.AddRaw(this,
+		&SMediaSourceManagerPreviews::SetManager);
+}
+
+SMediaSourceManagerPreviews::~SMediaSourceManagerPreviews()
+{
+	UMediaSourceManagerSettings* Settings = GetMutableDefault<UMediaSourceManagerSettings>();
+	if (Settings != nullptr)
+	{
+		Settings->OnManagerChanged.RemoveAll(this);
+	}
+}
+
+void SMediaSourceManagerPreviews::SetManager()
+{
+	MediaSourceManagerPtr = (GetDefault<UMediaSourceManagerSettings>()->GetManager());
+	RefreshChannels();
+}
+
+void SMediaSourceManagerPreviews::RefreshChannels()
+{
+	ChannelsContainer->ClearChildren();
+	UMediaSourceManager* MediaSourceManager = MediaSourceManagerPtr.Get();
+	if ((MediaSourceManager != nullptr) && (Style.IsValid()))
+	{
+		TSharedRef<ISlateStyle> StyleRef = Style.ToSharedRef();
+
+		// Loop over each channel.
+		int32 NumChannels = MediaSourceManager->Channels.Num();
+		for (int32 Index = 0; Index < NumChannels; ++Index)
+		{
+			UMediaSourceManagerChannel* Channel = MediaSourceManager->Channels[Index];
+			if (Channel != nullptr)
+			{
+				UMediaPlayer* MediaPlayer = Channel->GetMediaPlayer();
+				UMediaTexture* MediaTexture = Cast<UMediaTexture>(Channel->OutTexture);
+
+				// Add preview widget.
+				ChannelsContainer->AddSlot()
+					.Padding(2)
+					.HAlign(HAlign_Center)
+					[
+						SNew(SMediaPlayerEditorViewer, *MediaPlayer, MediaTexture, StyleRef, false)
+					];
+
+				Channel->Play();
+			}
+		}
+	}
+}
+
+
+#undef LOCTEXT_NAMESPACE
diff --git a/Engine/Plugins/Media/MediaSourceManager/Source/MediaSourceManagerEditor/Private/Widgets/SMediaSourceManagerPreviews.h b/Engine/Plugins/Media/MediaSourceManager/Source/MediaSourceManagerEditor/Private/Widgets/SMediaSourceManagerPreviews.h
new file mode 100644
--- /dev/null
+++ b/Engine/Plugins/Media/MediaSourceManager/Source/MediaSourceManagerEditor/Private/Widgets/SMediaSourceManagerPreviews.h
@@ -0,0 +1,49 @@
+// Copyright Epic Games, Inc. All Rights Reserved.
+
+#pragma once
+
+#include "Widgets/DeclarativeSyntaxSupport.h"
+#include "Widgets/SCompoundWidget.h"
+
+class ISlateStyle;
+class SHorizontalBox;
+class UMediaSourceManager;
+
+/**
+ * Implements a collection of previews for the media source manager.
+ */
+class SMediaSourceManagerPreviews
+	: public SCompoundWidget
+{
+public:
+
+	SLATE_BEGIN_ARGS(SMediaSourceManagerPreviews) { }
+	SLATE_END_ARGS()
+	~SMediaSourceManagerPreviews();
+
+	/**
+	 * Construct this widget.
+	 *
+	 * @param InArgs					The declaration data for this widget.
+	 */
+	void Construct(const FArguments& InArgs);
+
+private:
+	/**
+	 * Call this to use the manager in the settings.
+	 */
+	void SetManager();
+
+	/**
+	 * Update the channel widgets to reflect the channels in the current manager.
+	 */
+	void RefreshChannels();
+
+	/** Pointer to the object that is being viewed. */
+	TWeakObjectPtr<UMediaSourceManager> MediaSourceManagerPtr;
+
+	/** Pointer to the container for the channels. */
+	TSharedPtr<SHorizontalBox> ChannelsContainer;
+	/** Style to pass to the widgets. */
+	TSharedPtr<ISlateStyle> Style;
+};
diff --git a/Engine/Plugins/Media/MediaSourceManager/Source/MediaSourceManagerEditor/Private/Widgets/SMediaSourceManagerSources.cpp b/Engine/Plugins/Media/MediaSourceManager/Source/MediaSourceManagerEditor/Private/Widgets/SMediaSourceManagerSources.cpp
--- a/Engine/Plugins/Media/MediaSourceManager/Source/MediaSourceManagerEditor/Private/Widgets/SMediaSourceManagerSources.cpp
+++ b/Engine/Plugins/Media/MediaSourceManager/Source/MediaSourceManagerEditor/Private/Widgets/SMediaSourceManagerSources.cpp
@@ -11,9 +11,9 @@
 #define LOCTEXT_NAMESPACE "SMediaSourceManagerSources"
 
 void SMediaSourceManagerSources::Construct(const FArguments& InArgs,
-	UMediaSourceManager& InMediaSourceManager)
+	UMediaSourceManager* InMediaSourceManager)
 {
-	MediaSourceManagerPtr = &InMediaSourceManager;
+	MediaSourceManagerPtr = InMediaSourceManager;
 
 	ChildSlot
 		[
@@ -30,6 +30,12 @@ void SMediaSourceManagerSources::Construct(const FArguments& InArgs,
 	RefreshChannels();
 }
 
+void SMediaSourceManagerSources::SetManager(UMediaSourceManager* InManager)
+{
+	MediaSourceManagerPtr = InManager;
+	RefreshChannels();
+}
+
 void SMediaSourceManagerSources::RefreshChannels()
 {
 	ChannelsContainer->ClearChildren();
diff --git a/Engine/Plugins/Media/MediaSourceManager/Source/MediaSourceManagerEditor/Private/Widgets/SMediaSourceManagerSources.h b/Engine/Plugins/Media/MediaSourceManager/Source/MediaSourceManagerEditor/Private/Widgets/SMediaSourceManagerSources.h
--- a/Engine/Plugins/Media/MediaSourceManager/Source/MediaSourceManagerEditor/Private/Widgets/SMediaSourceManagerSources.h
+++ b/Engine/Plugins/Media/MediaSourceManager/Source/MediaSourceManagerEditor/Private/Widgets/SMediaSourceManagerSources.h
@@ -26,7 +26,12 @@ public:
 	 * @param InArgs					The declaration data for this widget.
 	 * @param InMediaSourceManager		The manager to show details for.
 	 */
-	void Construct(const FArguments& InArgs, UMediaSourceManager& InMediaSourceManager);
+	void Construct(const FArguments& InArgs, UMediaSourceManager* InMediaSourceManager);
+
+	/**
+	 * Call this to set which manager to use.
+	 */
+	void SetManager(UMediaSourceManager* InManager);
 
 private:
 
