MediaSourceManager: Url does not show for channel previews.
#rnx

diff --git a/Engine/Plugins/Media/MediaPlayerEditor/Source/MediaPlayerEditor/Private/Widgets/SMediaPlayerEditorViewer.cpp b/Engine/Plugins/Media/MediaPlayerEditor/Source/MediaPlayerEditor/Private/Widgets/SMediaPlayerEditorViewer.cpp
--- a/Engine/Plugins/Media/MediaPlayerEditor/Source/MediaPlayerEditor/Private/Widgets/SMediaPlayerEditorViewer.cpp
+++ b/Engine/Plugins/Media/MediaPlayerEditor/Source/MediaPlayerEditor/Private/Widgets/SMediaPlayerEditorViewer.cpp
@@ -257,6 +257,7 @@ void SMediaPlayerEditorViewer::Construct(const FArguments& InArgs, UMediaPlayer&
 					[
 						// url area
 						SNew(SHorizontalBox)
+							.Visibility(InArgs._bShowUrl ? EVisibility::Visible : EVisibility::Collapsed)
 
 						+ SHorizontalBox::Slot()
 							.AutoWidth()
diff --git a/Engine/Plugins/Media/MediaPlayerEditor/Source/MediaPlayerEditor/Public/Widgets/SMediaPlayerEditorViewer.h b/Engine/Plugins/Media/MediaPlayerEditor/Source/MediaPlayerEditor/Public/Widgets/SMediaPlayerEditorViewer.h
--- a/Engine/Plugins/Media/MediaPlayerEditor/Source/MediaPlayerEditor/Public/Widgets/SMediaPlayerEditorViewer.h
+++ b/Engine/Plugins/Media/MediaPlayerEditor/Source/MediaPlayerEditor/Public/Widgets/SMediaPlayerEditorViewer.h
@@ -30,7 +30,13 @@ class MEDIAPLAYEREDITOR_API SMediaPlayerEditorViewer
 {
 public:
 
-	SLATE_BEGIN_ARGS(SMediaPlayerEditorViewer) { }
+	SLATE_BEGIN_ARGS(SMediaPlayerEditorViewer)
+		: _bShowUrl(true)
+	{ }
+
+		/** If the url should be shown. */
+		SLATE_ARGUMENT(bool, bShowUrl)
+
 	SLATE_END_ARGS()
 
 public:
diff --git a/Engine/Plugins/Media/MediaSourceManager/Source/MediaSourceManagerEditor/Private/Toolkits/MediaSourceManagerEditorToolkit.cpp b/Engine/Plugins/Media/MediaSourceManager/Source/MediaSourceManagerEditor/Private/Toolkits/MediaSourceManagerEditorToolkit.cpp
--- a/Engine/Plugins/Media/MediaSourceManager/Source/MediaSourceManagerEditor/Private/Toolkits/MediaSourceManagerEditorToolkit.cpp
+++ b/Engine/Plugins/Media/MediaSourceManager/Source/MediaSourceManagerEditor/Private/Toolkits/MediaSourceManagerEditorToolkit.cpp
@@ -219,7 +219,8 @@ TSharedRef<SDockTab> FMediaSourceManagerEditorToolkit::HandlePreviewTabManagerSp
 			{
 				UMediaPlayer* MediaPlayer = Channel->GetMediaPlayer();
 				UMediaTexture* MediaTexture = Cast<UMediaTexture>(Channel->OutTexture);
-				TabWidget = SNew(SMediaPlayerEditorViewer, *MediaPlayer, MediaTexture, Style, false);
+				TabWidget = SNew(SMediaPlayerEditorViewer, *MediaPlayer, MediaTexture, Style, false)
+					.bShowUrl(false);
 			}
 		}
 	}
diff --git a/Engine/Plugins/Media/MediaSourceManager/Source/MediaSourceManagerEditor/Private/Widgets/SMediaSourceManagerPreviews.cpp b/Engine/Plugins/Media/MediaSourceManager/Source/MediaSourceManagerEditor/Private/Widgets/SMediaSourceManagerPreviews.cpp
--- a/Engine/Plugins/Media/MediaSourceManager/Source/MediaSourceManagerEditor/Private/Widgets/SMediaSourceManagerPreviews.cpp
+++ b/Engine/Plugins/Media/MediaSourceManager/Source/MediaSourceManagerEditor/Private/Widgets/SMediaSourceManagerPreviews.cpp
@@ -71,6 +71,7 @@ void SMediaSourceManagerPreviews::RefreshChannels()
 					.HAlign(HAlign_Center)
 					[
 						SNew(SMediaPlayerEditorViewer, *MediaPlayer, MediaTexture, StyleRef, false)
+							.bShowUrl(false)
 					];
 
 				Channel->Play();
