MediaSourceManager: Added button to assign an input to a channel. Hooked up MediaIO devices as possible inputs.

#rnx

diff --git a/Engine/Plugins/Media/BlackmagicMedia/Source/BlackmagicMedia/Private/Blackmagic/BlackmagicDeviceProvider.cpp b/Engine/Plugins/Media/BlackmagicMedia/Source/BlackmagicMedia/Private/Blackmagic/BlackmagicDeviceProvider.cpp
--- a/Engine/Plugins/Media/BlackmagicMedia/Source/BlackmagicMedia/Private/Blackmagic/BlackmagicDeviceProvider.cpp
+++ b/Engine/Plugins/Media/BlackmagicMedia/Source/BlackmagicMedia/Private/Blackmagic/BlackmagicDeviceProvider.cpp
@@ -5,6 +5,7 @@
 #include "Blackmagic.h"
 #include "BlackmagicLib.h"
 #include "BlackmagicMediaPrivate.h"
+#include "BlackmagicMediaSource.h"
 #include "IBlackmagicMediaModule.h"
 #include "MediaIOCoreCommonDisplayMode.h"
 
@@ -431,6 +432,19 @@ FMediaIOOutputConfiguration FBlackmagicDeviceProvider::GetDefaultOutputConfigura
 }
 
 
+UMediaSource* FBlackmagicDeviceProvider::CreateMediaSource(
+	const FMediaIOConfiguration& InConfiguration, UObject* Outer) const
+{
+	UBlackmagicMediaSource* MediaSource = NewObject<UBlackmagicMediaSource>(Outer);
+	if (MediaSource != nullptr)
+	{
+		MediaSource->MediaConfiguration = InConfiguration;
+	}
+
+	return MediaSource;
+}
+
+
 FText FBlackmagicDeviceProvider::ToText(const FMediaIOConfiguration& InConfiguration) const
 {
 	if (InConfiguration.IsValid())
diff --git a/Engine/Plugins/Media/BlackmagicMedia/Source/BlackmagicMedia/Public/BlackmagicDeviceProvider.h b/Engine/Plugins/Media/BlackmagicMedia/Source/BlackmagicMedia/Public/BlackmagicDeviceProvider.h
--- a/Engine/Plugins/Media/BlackmagicMedia/Source/BlackmagicMedia/Public/BlackmagicDeviceProvider.h
+++ b/Engine/Plugins/Media/BlackmagicMedia/Source/BlackmagicMedia/Public/BlackmagicDeviceProvider.h
@@ -46,6 +46,9 @@ public:
 	virtual FMediaIOOutputConfiguration GetDefaultOutputConfiguration() const override;
 	virtual FMediaIOVideoTimecodeConfiguration GetDefaultTimecodeConfiguration() const override;
 
+	virtual UMediaSource* CreateMediaSource(const FMediaIOConfiguration& InConfiguration,
+		UObject* Outer) const override;
+
 	virtual FText ToText(const FMediaIOConfiguration& InConfiguration) const override;
 	virtual FText ToText(const FMediaIOConnection& InConnection) const override;
 	virtual FText ToText(const FMediaIOOutputConfiguration& InConfiguration) const override;
diff --git a/Engine/Plugins/Media/MediaIOFramework/Source/MediaIOCore/Public/IMediaIOCoreDeviceProvider.h b/Engine/Plugins/Media/MediaIOFramework/Source/MediaIOCore/Public/IMediaIOCoreDeviceProvider.h
--- a/Engine/Plugins/Media/MediaIOFramework/Source/MediaIOCore/Public/IMediaIOCoreDeviceProvider.h
+++ b/Engine/Plugins/Media/MediaIOFramework/Source/MediaIOCore/Public/IMediaIOCoreDeviceProvider.h
@@ -5,6 +5,8 @@
 #include "CoreMinimal.h"
 #include "MediaIOCoreDefinitions.h"
 
+class UMediaSource;
+
 class MEDIAIOCORE_API IMediaIOCoreDeviceProvider
 {
 public:
@@ -25,6 +27,9 @@ public:
 	virtual FMediaIOOutputConfiguration GetDefaultOutputConfiguration() const = 0;
 	virtual FMediaIOVideoTimecodeConfiguration GetDefaultTimecodeConfiguration() const = 0;
 
+	virtual UMediaSource* CreateMediaSource(const FMediaIOConfiguration& InConfiguration,
+		UObject* Outer) const { return nullptr; }
+
 	virtual FText ToText(const FMediaIOConfiguration& InConfiguration) const;
 	virtual FText ToText(const FMediaIOConnection& InConnection) const;
 	virtual FText ToText(const FMediaIODevice& InDevice) const;
diff --git a/Engine/Plugins/Media/MediaSourceManager/Source/MediaSourceManagerEditor/MediaSourceManagerEditor.Build.cs b/Engine/Plugins/Media/MediaSourceManager/Source/MediaSourceManagerEditor/MediaSourceManagerEditor.Build.cs
--- a/Engine/Plugins/Media/MediaSourceManager/Source/MediaSourceManagerEditor/MediaSourceManagerEditor.Build.cs
+++ b/Engine/Plugins/Media/MediaSourceManager/Source/MediaSourceManagerEditor/MediaSourceManagerEditor.Build.cs
@@ -15,6 +15,7 @@ public class MediaSourceManagerEditor : ModuleRules
 				"CoreUObject",
 				"InputCore",
 				"MediaAssets",
+				"MediaIOCore",
 				"MediaPlayerEditor",
 				"MediaSourceManager",
 				"PropertyEditor",
diff --git a/Engine/Plugins/Media/MediaSourceManager/Source/MediaSourceManagerEditor/Private/Widgets/SMediaSourceManagerChannel.cpp b/Engine/Plugins/Media/MediaSourceManager/Source/MediaSourceManagerEditor/Private/Widgets/SMediaSourceManagerChannel.cpp
--- a/Engine/Plugins/Media/MediaSourceManager/Source/MediaSourceManagerEditor/Private/Widgets/SMediaSourceManagerChannel.cpp
+++ b/Engine/Plugins/Media/MediaSourceManager/Source/MediaSourceManagerEditor/Private/Widgets/SMediaSourceManagerChannel.cpp
@@ -3,15 +3,26 @@
 #include "Widgets/SMediaSourceManagerChannel.h"
 
 #include "DragAndDrop/AssetDragDropOp.h"
+#include "Framework/MultiBox/MultiBoxBuilder.h"
+#include "Framework/Notifications/NotificationManager.h"
+#include "IMediaIOCoreDeviceProvider.h"
+#include "IMediaIOCoreModule.h"
 #include "Inputs/MediaSourceManagerInputMediaSource.h"
 #include "MediaSource.h"
 #include "MediaSourceManagerChannel.h"
+#include "Widgets/Input/SComboButton.h"
 #include "Widgets/Layout/SScrollBox.h"
+#include "Widgets/Notifications/SNotificationList.h"
 #include "Widgets/SMediaSourceManagerTexture.h"
 #include "Widgets/Text/STextBlock.h"
 
 #define LOCTEXT_NAMESPACE "SMediaSourceManagerChannel"
 
+SMediaSourceManagerChannel::~SMediaSourceManagerChannel()
+{
+	DismissErrorNotification();
+}
+
 void SMediaSourceManagerChannel::Construct(const FArguments& InArgs,
 	UMediaSourceManagerChannel* InChannel)
 {
@@ -31,6 +42,24 @@ void SMediaSourceManagerChannel::Construct(const FArguments& InArgs,
 						.Text(FText::FromString(ChannelPtr->Name))
 				]
 
+			// Set input.
+			+ SHorizontalBox::Slot()
+				.AutoWidth()
+				.Padding(2)
+				.HAlign(HAlign_Left)
+				[
+					SNew(SComboButton)
+						.OnGetMenuContent(this, &SMediaSourceManagerChannel::CreateAssignInputMenu)
+						.ContentPadding(2)
+						.ButtonContent()
+						[
+							SNew(STextBlock)
+								.ToolTipText(LOCTEXT("Assign_ToolTip",
+									"Assign an input to this channel."))
+								.Text(LOCTEXT("AssignInput", "Assign Input"))
+						]
+				]
+
 			// Name of input.
 			+ SHorizontalBox::Slot()
 				.FillWidth(0.11f)
@@ -110,6 +139,100 @@ FReply SMediaSourceManagerChannel::OnDrop(const FGeometry& MyGeometry, const FDr
 	return FReply::Unhandled();
 }
 
+
+TSharedRef<SWidget> SMediaSourceManagerChannel::CreateAssignInputMenu()
+{
+	FMenuBuilder MenuBuilder(true, NULL);
+
+	// Get all Media IO device providers.
+	IMediaIOCoreModule& MediaIOCoreModule = IMediaIOCoreModule::Get();
+	TConstArrayView<IMediaIOCoreDeviceProvider*> DeviceProviders =
+		MediaIOCoreModule.GetDeviceProviders();
+
+	// Loop through each provider.
+	for (IMediaIOCoreDeviceProvider* DeviceProvider : DeviceProviders)
+	{
+		if (DeviceProvider != nullptr)
+		{
+			// Start menu section.
+			FName ProviderName = DeviceProvider->GetFName();
+			MenuBuilder.BeginSection(ProviderName, FText::FromName(ProviderName));
+
+			// Go over all input configurations.
+			TArray<FMediaIOConfiguration> Configurations = DeviceProvider->GetConfigurations(true, false);
+			for (const FMediaIOConfiguration& Configuration : Configurations)
+			{
+				// Add this device.
+				FName DeviceName = Configuration.MediaConnection.Device.DeviceName;
+
+				FUIAction AssignMediaIOInputAction(FExecuteAction::CreateSP(this,
+					&SMediaSourceManagerChannel::AssignMediaIOInput, DeviceProvider, Configuration));
+				MenuBuilder.AddMenuEntry(FText::FromName(DeviceName), FText(), FSlateIcon(),
+					AssignMediaIOInputAction);
+			}
+
+			MenuBuilder.EndSection();
+		}
+	}
+
+	return MenuBuilder.MakeWidget();
+}
+
+void SMediaSourceManagerChannel::AssignMediaIOInput(IMediaIOCoreDeviceProvider* DeviceProvider,
+	FMediaIOConfiguration Config)
+{
+	UMediaSourceManagerChannel* Channel = ChannelPtr.Get();
+	if (Channel != nullptr)
+	{
+		// Create media source.
+		UMediaSource* MediaSource = DeviceProvider->CreateMediaSource(Config, Channel);
+		if (MediaSource != nullptr)
+		{
+			// Assign to channel.
+			Channel->Modify();
+			UMediaSourceManagerInputMediaSource* Input = NewObject<UMediaSourceManagerInputMediaSource>(Channel);
+			Input->MediaSource = MediaSource;
+			Channel->Input = Input;
+
+			Channel->Play();
+
+			Refresh();
+		}
+		else
+		{
+			// Failed to create media source.
+			// Remove any existing error.
+			DismissErrorNotification();
+
+			// Inform the user.
+			FNotificationInfo Info(LOCTEXT("FailedToCreateMediaSource", "Failed to create a Media Source."));
+			Info.bFireAndForget = false;
+			Info.bUseLargeFont = false;
+			Info.bUseThrobber = false;
+			Info.FadeOutDuration = 0.25f;
+			Info.ButtonDetails.Add(FNotificationButtonInfo(LOCTEXT("Dismiss", "Dismiss"), LOCTEXT("DismissToolTip", "Dismiss this notification."),
+				FSimpleDelegate::CreateSP(this, &SMediaSourceManagerChannel::DismissErrorNotification)));
+
+			ErrorNotificationPtr = FSlateNotificationManager::Get().AddNotification(Info);
+			TSharedPtr<SNotificationItem> ErrorNotification = ErrorNotificationPtr.Pin();
+			if (ErrorNotification != nullptr)
+			{
+				ErrorNotification->SetCompletionState(SNotificationItem::CS_Pending);
+			}
+		}
+	}
+}
+
+void SMediaSourceManagerChannel::DismissErrorNotification()
+{
+	TSharedPtr<SNotificationItem> ErrorNotification = ErrorNotificationPtr.Pin();
+	if (ErrorNotification != nullptr)
+	{
+		ErrorNotification->ExpireAndFadeout();
+		ErrorNotification.Reset();
+	}
+}
+
 void SMediaSourceManagerChannel::Refresh()
 {
 	// Get channel.
diff --git a/Engine/Plugins/Media/MediaSourceManager/Source/MediaSourceManagerEditor/Private/Widgets/SMediaSourceManagerChannel.h b/Engine/Plugins/Media/MediaSourceManager/Source/MediaSourceManagerEditor/Private/Widgets/SMediaSourceManagerChannel.h
--- a/Engine/Plugins/Media/MediaSourceManager/Source/MediaSourceManagerEditor/Private/Widgets/SMediaSourceManagerChannel.h
+++ b/Engine/Plugins/Media/MediaSourceManager/Source/MediaSourceManagerEditor/Private/Widgets/SMediaSourceManagerChannel.h
@@ -5,8 +5,11 @@
 #include "Widgets/DeclarativeSyntaxSupport.h"
 #include "Widgets/SCompoundWidget.h"
 
+class IMediaIOCoreDeviceProvider;
+class SNotificationItem;
 class STextBlock;
 class UMediaSourceManagerChannel;
+struct FMediaIOConfiguration;
 
 /**
  * Implements a single channel.
@@ -19,6 +22,8 @@ public:
 	SLATE_BEGIN_ARGS(SMediaSourceManagerChannel) { }
 	SLATE_END_ARGS()
 
+	virtual ~SMediaSourceManagerChannel();
+
 	/**
 	 * Construct this widget.
 	 *
@@ -41,7 +46,27 @@ private:
 	TWeakObjectPtr<UMediaSourceManagerChannel> ChannelPtr;
 	/** Widget for the input name. */
 	TSharedPtr<STextBlock> InputNameTextBlock;
+	/** Holds our error notification so we can dismiss it. */
+	TWeakPtr<SNotificationItem> ErrorNotificationPtr;
 	
+	/**
+	 * Callback to create the assign input menu.
+	 */
+	TSharedRef<SWidget> CreateAssignInputMenu();
+
+	/**
+	 * Assigns a MediaIO input to this channel.
+	 *
+	 * @param DeviceProvider	MediaIO device provider to get the input from.
+	 * @param Config			MediaIO configuration to pass to the device provider.
+	 */
+	void AssignMediaIOInput(IMediaIOCoreDeviceProvider* DeviceProvider, FMediaIOConfiguration Config);
+
+	/**
+	 * Call this to remove the error notification.
+	 */
+	void DismissErrorNotification();
+
 	/**
 	 * Refreshes the widgets based on the current sstate.
 	 */
