MediaSourceManager: Added media sources to the assign input button menu.

diff --git a/Engine/Plugins/Media/MediaSourceManager/Source/MediaSourceManagerEditor/Private/Widgets/SMediaSourceManagerChannel.cpp b/Engine/Plugins/Media/MediaSourceManager/Source/MediaSourceManagerEditor/Private/Widgets/SMediaSourceManagerChannel.cpp
--- a/Engine/Plugins/Media/MediaSourceManager/Source/MediaSourceManagerEditor/Private/Widgets/SMediaSourceManagerChannel.cpp
+++ b/Engine/Plugins/Media/MediaSourceManager/Source/MediaSourceManagerEditor/Private/Widgets/SMediaSourceManagerChannel.cpp
@@ -2,9 +2,12 @@
 
 #include "Widgets/SMediaSourceManagerChannel.h"
 
+#include "ContentBrowserModule.h"
 #include "DragAndDrop/AssetDragDropOp.h"
+#include "Framework/Application/SlateApplication.h"
 #include "Framework/MultiBox/MultiBoxBuilder.h"
 #include "Framework/Notifications/NotificationManager.h"
+#include "IContentBrowserSingleton.h"
 #include "IMediaIOCoreDeviceProvider.h"
 #include "IMediaIOCoreModule.h"
 #include "Inputs/MediaSourceManagerInputMediaSource.h"
@@ -112,27 +115,17 @@ FReply SMediaSourceManagerChannel::OnDrop(const FGeometry& MyGeometry, const FDr
 	// Is this an asset drop?
 	if (TSharedPtr<FAssetDragDropOp> AssetDragDrop = DragDropEvent.GetOperationAs<FAssetDragDropOp>())
 	{
-		UMediaSourceManagerChannel* Channel = ChannelPtr.Get();
-		if (Channel != nullptr)
+		for (const FAssetData& Asset : AssetDragDrop->GetAssets())
 		{
-			for (const FAssetData& Asset : AssetDragDrop->GetAssets())
+			// Is this a media source?
+			UMediaSource* MediaSource = Cast<UMediaSource>(Asset.GetAsset());
+			if (MediaSource != nullptr)
 			{
-				// Is this a media source?
-				UMediaSource* MediaSource = Cast<UMediaSource>(Asset.GetAsset());
-				if (MediaSource != nullptr)
-				{
-					Channel->Modify();
-					UMediaSourceManagerInputMediaSource* Input = NewObject<UMediaSourceManagerInputMediaSource>(Channel);
-					Input->MediaSource = MediaSource;
-					Channel->Input = Input;
-
-					Channel->Play();
-						
-					Refresh();
-					break;
-				}
+				AssignMediaSourceInput(MediaSource);
+				break;
 			}
 		}
+		
 		return FReply::Handled();
 	}
 
@@ -175,9 +168,90 @@ TSharedRef<SWidget> SMediaSourceManagerChannel::CreateAssignInputMenu()
 		}
 	}
 
+	// Add assets.
+	MenuBuilder.BeginSection(NAME_None, LOCTEXT("MediaSourceAssets", "Media Source Assets"));
+	auto SubMenuCallback = [this](FMenuBuilder& SubMenuBuilder)
+	{
+		SubMenuBuilder.AddWidget(BuildMediaSourcePickerWidget(), FText::GetEmpty(), true);
+	};
+	MenuBuilder.AddSubMenu(
+		LOCTEXT("SelectAsset", "Select Asset"),
+		LOCTEXT("SelectAsset_ToolTip", "Select an existing Media Source asset."),
+		FNewMenuDelegate::CreateLambda(SubMenuCallback)
+	);
+	MenuBuilder.EndSection();
+
 	return MenuBuilder.MakeWidget();
 }
 
+TSharedRef<SWidget> SMediaSourceManagerChannel::BuildMediaSourcePickerWidget()
+{
+	FAssetPickerConfig AssetPickerConfig;
+	{
+		AssetPickerConfig.OnAssetSelected = FOnAssetSelected::CreateRaw(this,
+			&SMediaSourceManagerChannel::AddMediaSource);
+		AssetPickerConfig.OnAssetEnterPressed = FOnAssetEnterPressed::CreateRaw(this,
+			&SMediaSourceManagerChannel::AddMediaSourceEnterPressed);
+		AssetPickerConfig.bAllowNullSelection = false;
+		AssetPickerConfig.InitialAssetViewType = EAssetViewType::List;
+		AssetPickerConfig.Filter.bRecursiveClasses = true;
+		AssetPickerConfig.Filter.ClassPaths.Add(UMediaSource::StaticClass()->GetClassPathName());
+		AssetPickerConfig.SaveSettingsName = TEXT("MediaSourceManagerAssetPicker");
+	}
+
+	FContentBrowserModule& ContentBrowserModule = FModuleManager::Get().LoadModuleChecked<FContentBrowserModule>(TEXT("ContentBrowser"));
+
+	TSharedRef<SBox> Picker = SNew(SBox)
+		.WidthOverride(300.0f)
+		.HeightOverride(300.f)
+		[
+			ContentBrowserModule.Get().CreateAssetPicker(AssetPickerConfig)
+		];
+
+	return Picker;
+}
+
+void SMediaSourceManagerChannel::AddMediaSource(const FAssetData& AssetData)
+{
+	FSlateApplication::Get().DismissAllMenus();
+
+	// Get media source from the asset..
+	UObject* SelectedObject = AssetData.GetAsset();
+	if (SelectedObject)
+	{
+		UMediaSource* MediaSource = Cast<UMediaSource>(AssetData.GetAsset());
+		if (MediaSource != nullptr)
+		{
+			AssignMediaSourceInput(MediaSource);
+		}
+	}
+}
+
+void SMediaSourceManagerChannel::AddMediaSourceEnterPressed(const TArray<FAssetData>& AssetData)
+{
+	if (AssetData.Num() > 0)
+	{
+		AddMediaSource(AssetData[0].GetAsset());
+	}
+}
+
+
+void SMediaSourceManagerChannel::AssignMediaSourceInput(UMediaSource* MediaSource)
+{
+	UMediaSourceManagerChannel* Channel = ChannelPtr.Get();
+	if (Channel != nullptr)
+	{
+		// Assign to channel.
+		Channel->Modify();
+		UMediaSourceManagerInputMediaSource* Input = NewObject<UMediaSourceManagerInputMediaSource>(Channel);
+		Input->MediaSource = MediaSource;
+		Channel->Input = Input;
+		Channel->Play();
+
+		Refresh();
+	}
+}
+
 void SMediaSourceManagerChannel::AssignMediaIOInput(IMediaIOCoreDeviceProvider* DeviceProvider,
 	FMediaIOConfiguration Config)
 {
@@ -188,15 +262,7 @@ void SMediaSourceManagerChannel::AssignMediaIOInput(IMediaIOCoreDeviceProvider*
 		UMediaSource* MediaSource = DeviceProvider->CreateMediaSource(Config, Channel);
 		if (MediaSource != nullptr)
 		{
-			// Assign to channel.
-			Channel->Modify();
-			UMediaSourceManagerInputMediaSource* Input = NewObject<UMediaSourceManagerInputMediaSource>(Channel);
-			Input->MediaSource = MediaSource;
-			Channel->Input = Input;
-
-			Channel->Play();
-
-			Refresh();
+			AssignMediaSourceInput(MediaSource);
 		}
 		else
 		{
diff --git a/Engine/Plugins/Media/MediaSourceManager/Source/MediaSourceManagerEditor/Private/Widgets/SMediaSourceManagerChannel.h b/Engine/Plugins/Media/MediaSourceManager/Source/MediaSourceManagerEditor/Private/Widgets/SMediaSourceManagerChannel.h
--- a/Engine/Plugins/Media/MediaSourceManager/Source/MediaSourceManagerEditor/Private/Widgets/SMediaSourceManagerChannel.h
+++ b/Engine/Plugins/Media/MediaSourceManager/Source/MediaSourceManagerEditor/Private/Widgets/SMediaSourceManagerChannel.h
@@ -8,7 +8,10 @@
 class IMediaIOCoreDeviceProvider;
 class SNotificationItem;
 class STextBlock;
+class UMediaSource;
 class UMediaSourceManagerChannel;
+
+struct FAssetData;
 struct FMediaIOConfiguration;
 
 /**
@@ -54,6 +57,28 @@ private:
 	 */
 	TSharedRef<SWidget> CreateAssignInputMenu();
 
+	/**
+	 * Creates a widget to select a media source.
+	 */
+	TSharedRef<SWidget> BuildMediaSourcePickerWidget();
+
+	/**
+	 * Callback to add a media source.
+	 */
+	void AddMediaSource(const FAssetData& AssetData);
+
+	/**
+	 * Callback to add a media source.
+	 */
+	void AddMediaSourceEnterPressed(const TArray<FAssetData>& AssetData);
+
+	/**
+	 * Assignss a media source to this channel.
+	 * 
+	 * @param MediaSource	Media source to use as input.
+	 */
+	void AssignMediaSourceInput(UMediaSource* MediaSource);
+
 	/**
 	 * Assigns a MediaIO input to this channel.
 	 *
