MediaSourceManager: Fixed the same material always being used when dragging a channel texture over an object.

#rnx

diff --git a/Engine/Plugins/Media/MediaSourceManager/Source/MediaSourceManager/Public/MediaSourceManagerChannel.h b/Engine/Plugins/Media/MediaSourceManager/Source/MediaSourceManager/Public/MediaSourceManagerChannel.h
--- a/Engine/Plugins/Media/MediaSourceManager/Source/MediaSourceManager/Public/MediaSourceManagerChannel.h
+++ b/Engine/Plugins/Media/MediaSourceManager/Source/MediaSourceManager/Public/MediaSourceManagerChannel.h
@@ -8,6 +8,7 @@
 
 #include "MediaSourceManagerChannel.generated.h"
 
+class UMaterialInstanceConstant;
 class UMediaPlayer;
 class UMediaSource;
 class UMediaSourceManagerInput;
@@ -42,6 +43,10 @@ public:
 	UPROPERTY(VisibleAnywhere, BlueprintReadOnly, Category = "Channel")
 	TObjectPtr<UTexture> OutTexture = nullptr;
 
+	/** Material that uses our texture. */
+	UPROPERTY()
+	TObjectPtr<UMaterialInstanceConstant> Material = nullptr;
+
 	/**
 	 * Call this to get the media player.
 	 */
diff --git a/Engine/Plugins/Media/MediaSourceManager/Source/MediaSourceManagerEditor/MediaSourceManagerEditor.Build.cs b/Engine/Plugins/Media/MediaSourceManager/Source/MediaSourceManagerEditor/MediaSourceManagerEditor.Build.cs
--- a/Engine/Plugins/Media/MediaSourceManager/Source/MediaSourceManagerEditor/MediaSourceManagerEditor.Build.cs
+++ b/Engine/Plugins/Media/MediaSourceManager/Source/MediaSourceManagerEditor/MediaSourceManagerEditor.Build.cs
@@ -13,6 +13,7 @@ public class MediaSourceManagerEditor : ModuleRules
 			{
 				"Core",
 				"CoreUObject",
+				"Engine",
 				"InputCore",
 				"MediaAssets",
 				"MediaIOCore",
diff --git a/Engine/Plugins/Media/MediaSourceManager/Source/MediaSourceManagerEditor/Private/Widgets/SMediaSourceManagerTexture.cpp b/Engine/Plugins/Media/MediaSourceManager/Source/MediaSourceManagerEditor/Private/Widgets/SMediaSourceManagerTexture.cpp
--- a/Engine/Plugins/Media/MediaSourceManager/Source/MediaSourceManagerEditor/Private/Widgets/SMediaSourceManagerTexture.cpp
+++ b/Engine/Plugins/Media/MediaSourceManager/Source/MediaSourceManagerEditor/Private/Widgets/SMediaSourceManagerTexture.cpp
@@ -6,9 +6,12 @@
 #include "AssetThumbnail.h"
 #include "DragAndDrop/AssetDragDropOp.h"
 #include "Engine/Texture.h"
+#include "Materials/MaterialInstanceConstant.h"
 #include "MediaSourceManagerChannel.h"
 #include "ThumbnailRendering/ThumbnailManager.h"
 
+FLazyName SMediaSourceManagerTexture::MediaTextureName("MediaTexture");
+
 void SMediaSourceManagerTexture::Construct(const FArguments& InArgs,
 	UMediaSourceManagerChannel* InChannel)
 {
@@ -48,10 +51,58 @@ FReply SMediaSourceManagerTexture::OnDragDetected(const FGeometry& MyGeometry, c
 		UTexture* Texture = Channel->OutTexture;
 		if (Texture != nullptr)
 		{
-			FAssetData AssetData(Texture);
-			return FReply::Handled().BeginDragDrop(FAssetDragDropOp::New(AssetData));
+			UMaterialInstanceConstant* Material = GetMaterial();
+			if (Material != nullptr)
+			{
+
+				FAssetData AssetData(Material);
+				return FReply::Handled().BeginDragDrop(FAssetDragDropOp::New(AssetData));
+			}
 		}
 	}
 
 	return FReply::Unhandled();
 }
+
+UMaterialInstanceConstant* SMediaSourceManagerTexture::GetMaterial()
+{
+	UMaterialInstanceConstant* MaterialInstance = nullptr;
+
+	UMediaSourceManagerChannel* Channel = ChannelPtr.Get();
+	if (Channel != nullptr)
+	{
+		// Do we already have a material?
+		MaterialInstance = Channel->Material;
+		if (MaterialInstance == nullptr)
+		{
+			// No. Create one.
+			UMaterial* Material = LoadObject<UMaterial>(NULL,
+				TEXT("/MediaSourceManager/M_MediaSourceManager"), NULL, LOAD_None, NULL);
+			if (Material != nullptr)
+			{
+				UObject* Outer = Channel->GetOuter();
+				Channel->Modify();
+				
+				FString MaterialName = Material->GetName();
+				if (MaterialName.StartsWith(TEXT("M_")))
+				{
+					MaterialName.InsertAt(1, TEXT("I"));
+				}
+				FName MaterialUniqueName = MakeUniqueObjectName(Outer, UMaterialInstanceConstant::StaticClass(),
+					FName(*MaterialName));
+
+				// Create instance.
+				MaterialInstance =
+					NewObject<UMaterialInstanceConstant>(Outer, MaterialUniqueName, RF_Public);
+				Channel->Material = MaterialInstance;
+				MaterialInstance->SetParentEditorOnly(Material);
+				MaterialInstance->SetTextureParameterValueEditorOnly(
+					FMaterialParameterInfo(MediaTextureName),
+					Channel->OutTexture);
+				MaterialInstance->PostEditChange();
+			}
+		}
+	}
+
+	return MaterialInstance;
+}
diff --git a/Engine/Plugins/Media/MediaSourceManager/Source/MediaSourceManagerEditor/Private/Widgets/SMediaSourceManagerTexture.h b/Engine/Plugins/Media/MediaSourceManager/Source/MediaSourceManagerEditor/Private/Widgets/SMediaSourceManagerTexture.h
--- a/Engine/Plugins/Media/MediaSourceManager/Source/MediaSourceManagerEditor/Private/Widgets/SMediaSourceManagerTexture.h
+++ b/Engine/Plugins/Media/MediaSourceManager/Source/MediaSourceManagerEditor/Private/Widgets/SMediaSourceManagerTexture.h
@@ -5,6 +5,7 @@
 #include "Widgets/DeclarativeSyntaxSupport.h"
 #include "Widgets/SCompoundWidget.h"
 
+class UMaterialInstanceConstant;
 class UMediaSourceManagerChannel;
 
 /**
@@ -33,7 +34,14 @@ protected:
 	// End of SWidget interface
 
 private:
+	/**
+	 * Gets a material for this texture.
+	 */
+	UMaterialInstanceConstant* GetMaterial();
 
 	/** Pointer to the object that is being viewed. */
 	TWeakObjectPtr<UMediaSourceManagerChannel> ChannelPtr;
+
+	/** Name for the media texture parameter in the material. */
+	static FLazyName MediaTextureName;
 };
