MediaSourceManager: Assign input button shows the current media source now.
#rnx

diff --git a/Engine/Plugins/Media/MediaSourceManager/Source/MediaSourceManagerEditor/Private/Widgets/SMediaSourceManagerChannel.cpp b/Engine/Plugins/Media/MediaSourceManager/Source/MediaSourceManagerEditor/Private/Widgets/SMediaSourceManagerChannel.cpp
--- a/Engine/Plugins/Media/MediaSourceManager/Source/MediaSourceManagerEditor/Private/Widgets/SMediaSourceManagerChannel.cpp
+++ b/Engine/Plugins/Media/MediaSourceManager/Source/MediaSourceManagerEditor/Private/Widgets/SMediaSourceManagerChannel.cpp
@@ -47,7 +47,7 @@ void SMediaSourceManagerChannel::Construct(const FArguments& InArgs,
 
 			// Set input.
 			+ SHorizontalBox::Slot()
-				.AutoWidth()
+				.FillWidth(0.6f)
 				.Padding(2)
 				.HAlign(HAlign_Left)
 				[
@@ -56,22 +56,12 @@ void SMediaSourceManagerChannel::Construct(const FArguments& InArgs,
 						.ContentPadding(2)
 						.ButtonContent()
 						[
-							SNew(STextBlock)
+							SAssignNew(InputNameTextBlock, STextBlock)
 								.ToolTipText(LOCTEXT("Assign_ToolTip",
 									"Assign an input to this channel."))
-								.Text(LOCTEXT("AssignInput", "Assign Input"))
 						]
 				]
 
-			// Name of input.
-			+ SHorizontalBox::Slot()
-				.FillWidth(0.11f)
-				.Padding(2)
-				.HAlign(HAlign_Left)
-				[
-					SAssignNew(InputNameTextBlock, STextBlock)
-				]
-
 			// Out texture
 			+ SHorizontalBox::Slot()
 				.FillWidth(0.11f)
@@ -305,17 +295,17 @@ void SMediaSourceManagerChannel::Refresh()
 	UMediaSourceManagerChannel* Channel = ChannelPtr.Get();
 	if (Channel != nullptr)
 	{
-		FString InputName = FString(TEXT("None"));
+		FText InputName = LOCTEXT("AssignInput", "Assign Input");
 
 		// Get input.
 		UMediaSourceManagerInput* Input = Channel->Input;
 		if (Input != nullptr)
 		{
-			InputName = Input->GetDisplayName();
+			InputName = FText::FromString(Input->GetDisplayName());
 		}
 
 		// Update input widgets.
-		InputNameTextBlock->SetText(FText::FromString(InputName));
+		InputNameTextBlock->SetText(InputName);
 	}
 }
 
