MediaSourceManager: Added button to edit the channel input.
#rnx

diff --git a/Engine/Plugins/Media/MediaSourceManager/Source/MediaSourceManagerEditor/Private/Widgets/SMediaSourceManagerChannel.cpp b/Engine/Plugins/Media/MediaSourceManager/Source/MediaSourceManagerEditor/Private/Widgets/SMediaSourceManagerChannel.cpp
--- a/Engine/Plugins/Media/MediaSourceManager/Source/MediaSourceManagerEditor/Private/Widgets/SMediaSourceManagerChannel.cpp
+++ b/Engine/Plugins/Media/MediaSourceManager/Source/MediaSourceManagerEditor/Private/Widgets/SMediaSourceManagerChannel.cpp
@@ -4,6 +4,7 @@
 
 #include "ContentBrowserModule.h"
 #include "DragAndDrop/AssetDragDropOp.h"
+#include "Editor.h"
 #include "Framework/Application/SlateApplication.h"
 #include "Framework/MultiBox/MultiBoxBuilder.h"
 #include "Framework/Notifications/NotificationManager.h"
@@ -13,6 +14,8 @@
 #include "Inputs/MediaSourceManagerInputMediaSource.h"
 #include "MediaSource.h"
 #include "MediaSourceManagerChannel.h"
+#include "Subsystems/AssetEditorSubsystem.h"
+#include "Widgets/Input/SButton.h"
 #include "Widgets/Input/SComboButton.h"
 #include "Widgets/Layout/SScrollBox.h"
 #include "Widgets/Notifications/SNotificationList.h"
@@ -62,6 +65,18 @@ void SMediaSourceManagerChannel::Construct(const FArguments& InArgs,
 						]
 				]
 
+			// Edit input.
+			+ SHorizontalBox::Slot()
+				.AutoWidth()
+				.HAlign(HAlign_Left)
+				.Padding(2)
+				[
+					SNew(SButton)
+						.ContentPadding(3)
+						.OnClicked(this, &SMediaSourceManagerChannel::OnEditInput)
+						.Text(LOCTEXT("EditInput", "Edit Input"))
+				]
+
 			// Out texture
 			+ SHorizontalBox::Slot()
 				.FillWidth(0.11f)
@@ -289,6 +304,34 @@ void SMediaSourceManagerChannel::DismissErrorNotification()
 	}
 }
 
+
+FReply SMediaSourceManagerChannel::OnEditInput()
+{
+	// Get our input.
+	TArray<UObject*> AssetArray;
+	UMediaSourceManagerChannel* Channel = ChannelPtr.Get();
+	if (Channel != nullptr)
+	{
+		UMediaSourceManagerInput* Input = Channel->Input;
+		if (Input != nullptr)
+		{
+			UMediaSource* MediaSource = Input->GetMediaSource();
+			if (MediaSource != nullptr)
+			{
+				AssetArray.Add(MediaSource);
+			}
+		}
+	}
+
+	// Open the editor.
+	if (AssetArray.Num() > 0)
+	{
+		GEditor->GetEditorSubsystem<UAssetEditorSubsystem>()->OpenEditorForAssets(AssetArray);
+	}
+
+	return FReply::Handled();
+}
+
 void SMediaSourceManagerChannel::Refresh()
 {
 	// Get channel.
diff --git a/Engine/Plugins/Media/MediaSourceManager/Source/MediaSourceManagerEditor/Private/Widgets/SMediaSourceManagerChannel.h b/Engine/Plugins/Media/MediaSourceManager/Source/MediaSourceManagerEditor/Private/Widgets/SMediaSourceManagerChannel.h
--- a/Engine/Plugins/Media/MediaSourceManager/Source/MediaSourceManagerEditor/Private/Widgets/SMediaSourceManagerChannel.h
+++ b/Engine/Plugins/Media/MediaSourceManager/Source/MediaSourceManagerEditor/Private/Widgets/SMediaSourceManagerChannel.h
@@ -92,6 +92,11 @@ private:
 	 */
 	void DismissErrorNotification();
 
+	/**
+	 * Called when the edit input button is pressed.
+	 */
+	FReply OnEditInput();
+
 	/**
 	 * Refreshes the widgets based on the current sstate.
 	 */
