MediaSourceManager: Channel replays its media if the media source is edited.

#rnx

diff --git a/Engine/Plugins/Media/MediaSourceManager/Source/MediaSourceManager/Private/MediaSourceManagerChannel.cpp b/Engine/Plugins/Media/MediaSourceManager/Source/MediaSourceManager/Private/MediaSourceManagerChannel.cpp
--- a/Engine/Plugins/Media/MediaSourceManager/Source/MediaSourceManager/Private/MediaSourceManagerChannel.cpp
+++ b/Engine/Plugins/Media/MediaSourceManager/Source/MediaSourceManager/Private/MediaSourceManagerChannel.cpp
@@ -63,6 +63,29 @@ void UMediaSourceManagerChannel::Play()
 	}
 }
 
+void UMediaSourceManagerChannel::PostInitProperties()
+{
+#if WITH_EDITOR
+	// Don't run on CDO
+	if (!HasAnyFlags(RF_ClassDefaultObject))
+	{
+		FCoreUObjectDelegates::OnObjectPropertyChanged.AddUObject(this,
+			&UMediaSourceManagerChannel::OnObjectPropertyChanged);
+	}
+#endif // WITH_EDITOR
+
+	Super::PostInitProperties();
+}
+
+void UMediaSourceManagerChannel::BeginDestroy()
+{
+#if WITH_EDITOR
+	FCoreUObjectDelegates::OnObjectPropertyChanged.RemoveAll(this);
+#endif // WITH_EDITOR
+
+	Super::BeginDestroy();
+}
+
 #if WITH_EDITOR
 
 void UMediaSourceManagerChannel::PostEditChangeProperty(FPropertyChangedEvent& PropertyChangedEvent)
@@ -79,6 +102,24 @@ void UMediaSourceManagerChannel::PostEditChangeProperty(FPropertyChangedEvent& P
 	}
 }
 
+void UMediaSourceManagerChannel::OnObjectPropertyChanged(UObject* ObjectBeingModified,
+	FPropertyChangedEvent& PropertyChangedEvent)
+{
+	// Did our media source change?
+	if (Input != nullptr)
+	{
+		UMediaSource* MediaSource = Input->GetMediaSource();
+		if (MediaSource != nullptr)
+		{
+			if (MediaSource == ObjectBeingModified)
+			{
+				// Restart playback.
+				Play();
+			}
+		}
+	}
+}
+
 #endif // WITH_EDITOR
 
 
diff --git a/Engine/Plugins/Media/MediaSourceManager/Source/MediaSourceManager/Public/MediaSourceManagerChannel.h b/Engine/Plugins/Media/MediaSourceManager/Source/MediaSourceManager/Public/MediaSourceManagerChannel.h
--- a/Engine/Plugins/Media/MediaSourceManager/Source/MediaSourceManager/Public/MediaSourceManagerChannel.h
+++ b/Engine/Plugins/Media/MediaSourceManager/Source/MediaSourceManager/Public/MediaSourceManagerChannel.h
@@ -62,6 +62,9 @@ public:
 	void Validate();
 
 	//~ Begin UObject interface
+	virtual void PostInitProperties() override;
+	virtual void BeginDestroy() override;
+
 #if WITH_EDITOR
 	virtual void PostEditChangeProperty(FPropertyChangedEvent& PropertyChangedEvent) override;
 #endif // WITH_EDITOR
@@ -72,4 +75,12 @@ private:
 	UPROPERTY()
 	TObjectPtr<UMediaPlayer> CurrentMediaPlayer = nullptr;
 
+#if WITH_EDITOR
+	
+	/**
+	 * Called when an object is edited.
+	 */
+	void OnObjectPropertyChanged(UObject* ObjectBeingModified, FPropertyChangedEvent& PropertyChangedEvent);
+
+#endif // WITH_EDITOR
 };
