MediaSourceManager: Assign input button just shows the capture card devices and not all configurations now.

#rnx

diff --git a/Engine/Plugins/Media/MediaSourceManager/Source/MediaSourceManagerEditor/Private/Widgets/SMediaSourceManagerChannel.cpp b/Engine/Plugins/Media/MediaSourceManager/Source/MediaSourceManagerEditor/Private/Widgets/SMediaSourceManagerChannel.cpp
--- a/Engine/Plugins/Media/MediaSourceManager/Source/MediaSourceManagerEditor/Private/Widgets/SMediaSourceManagerChannel.cpp
+++ b/Engine/Plugins/Media/MediaSourceManager/Source/MediaSourceManagerEditor/Private/Widgets/SMediaSourceManagerChannel.cpp
@@ -156,15 +156,15 @@ TSharedRef<SWidget> SMediaSourceManagerChannel::CreateAssignInputMenu()
 			FName ProviderName = DeviceProvider->GetFName();
 			MenuBuilder.BeginSection(ProviderName, FText::FromName(ProviderName));
 
-			// Go over all input configurations.
-			TArray<FMediaIOConfiguration> Configurations = DeviceProvider->GetConfigurations(true, false);
-			for (const FMediaIOConfiguration& Configuration : Configurations)
+			// Go over all devices.
+			TArray<FMediaIODevice> Devices = DeviceProvider->GetDevices();
+			for (const FMediaIODevice& Device : Devices)
 			{
 				// Add this device.
-				FName DeviceName = Configuration.MediaConnection.Device.DeviceName;
+				FName DeviceName = Device.DeviceName;
 
 				FUIAction AssignMediaIOInputAction(FExecuteAction::CreateSP(this,
-					&SMediaSourceManagerChannel::AssignMediaIOInput, DeviceProvider, Configuration));
+					&SMediaSourceManagerChannel::AssignMediaIOInput, DeviceProvider, Device));
 				MenuBuilder.AddMenuEntry(FText::FromName(DeviceName), FText(), FSlateIcon(),
 					AssignMediaIOInputAction);
 			}
@@ -258,13 +258,16 @@ void SMediaSourceManagerChannel::AssignMediaSourceInput(UMediaSource* MediaSourc
 }
 
 void SMediaSourceManagerChannel::AssignMediaIOInput(IMediaIOCoreDeviceProvider* DeviceProvider,
-	FMediaIOConfiguration Config)
+	FMediaIODevice Device)
 {
 	UMediaSourceManagerChannel* Channel = ChannelPtr.Get();
 	if (Channel != nullptr)
 	{
 		// Create media source.
-		UMediaSource* MediaSource = DeviceProvider->CreateMediaSource(Config, Channel);
+		FMediaIOConfiguration Configuration;
+		Configuration = DeviceProvider->GetDefaultConfiguration();
+		Configuration.MediaConnection.Device = Device;
+		UMediaSource* MediaSource = DeviceProvider->CreateMediaSource(Configuration, Channel);
 		if (MediaSource != nullptr)
 		{
 			AssignMediaSourceInput(MediaSource);
diff --git a/Engine/Plugins/Media/MediaSourceManager/Source/MediaSourceManagerEditor/Private/Widgets/SMediaSourceManagerChannel.h b/Engine/Plugins/Media/MediaSourceManager/Source/MediaSourceManagerEditor/Private/Widgets/SMediaSourceManagerChannel.h
--- a/Engine/Plugins/Media/MediaSourceManager/Source/MediaSourceManagerEditor/Private/Widgets/SMediaSourceManagerChannel.h
+++ b/Engine/Plugins/Media/MediaSourceManager/Source/MediaSourceManagerEditor/Private/Widgets/SMediaSourceManagerChannel.h
@@ -12,7 +12,7 @@ class UMediaSource;
 class UMediaSourceManagerChannel;
 
 struct FAssetData;
-struct FMediaIOConfiguration;
+struct FMediaIODevice;
 
 /**
  * Implements a single channel.
@@ -83,9 +83,9 @@ private:
 	 * Assigns a MediaIO input to this channel.
 	 *
 	 * @param DeviceProvider	MediaIO device provider to get the input from.
-	 * @param Config			MediaIO configuration to pass to the device provider.
+	 * @param Device			Device to pass to the device provider.
 	 */
-	void AssignMediaIOInput(IMediaIOCoreDeviceProvider* DeviceProvider, FMediaIOConfiguration Config);
+	void AssignMediaIOInput(IMediaIOCoreDeviceProvider* DeviceProvider, FMediaIODevice Device);
 
 	/**
 	 * Call this to remove the error notification.
