MediaSourceManager: Warning icon displays when an input has incorrect settings.

#rnx

diff --git a/Engine/Plugins/Media/MediaSourceManager/Source/MediaSourceManager/Private/MediaSourceManagerChannel.cpp b/Engine/Plugins/Media/MediaSourceManager/Source/MediaSourceManager/Private/MediaSourceManagerChannel.cpp
--- a/Engine/Plugins/Media/MediaSourceManager/Source/MediaSourceManager/Private/MediaSourceManagerChannel.cpp
+++ b/Engine/Plugins/Media/MediaSourceManager/Source/MediaSourceManager/Private/MediaSourceManagerChannel.cpp
@@ -115,6 +115,12 @@ void UMediaSourceManagerChannel::OnObjectPropertyChanged(UObject* ObjectBeingMod
 			{
 				// Restart playback.
 				Play();
+
+				// Inform everyone.
+				if (OnInputPropertyChanged.IsBound())
+				{
+					OnInputPropertyChanged.Broadcast();
+				}
 			}
 		}
 	}
diff --git a/Engine/Plugins/Media/MediaSourceManager/Source/MediaSourceManager/Public/MediaSourceManagerChannel.h b/Engine/Plugins/Media/MediaSourceManager/Source/MediaSourceManager/Public/MediaSourceManagerChannel.h
--- a/Engine/Plugins/Media/MediaSourceManager/Source/MediaSourceManager/Public/MediaSourceManagerChannel.h
+++ b/Engine/Plugins/Media/MediaSourceManager/Source/MediaSourceManager/Public/MediaSourceManagerChannel.h
@@ -15,6 +15,12 @@ class UMediaSourceManagerInput;
 class UProxyMediaSource;
 class UTexture;
 
+#if WITH_EDITOR
+
+DECLARE_EVENT(UMediaSourceManagerChannel, FOnInputPropertyChanged);
+
+#endif
+
 /**
 * Handles a single channel for the MediaSourceManager.
 */
@@ -47,6 +53,11 @@ public:
 	UPROPERTY()
 	TObjectPtr<UMaterialInstanceConstant> Material = nullptr;
 
+#if WITH_EDITOR
+	/** This is broadcast when something on the input is changed. */
+	FOnInputPropertyChanged OnInputPropertyChanged;
+#endif // WITH_EDITOR
+
 	/**
 	 * Call this to get the media player.
 	 */
diff --git a/Engine/Plugins/Media/MediaSourceManager/Source/MediaSourceManagerEditor/Private/Widgets/SMediaSourceManagerChannel.cpp b/Engine/Plugins/Media/MediaSourceManager/Source/MediaSourceManagerEditor/Private/Widgets/SMediaSourceManagerChannel.cpp
--- a/Engine/Plugins/Media/MediaSourceManager/Source/MediaSourceManagerEditor/Private/Widgets/SMediaSourceManagerChannel.cpp
+++ b/Engine/Plugins/Media/MediaSourceManager/Source/MediaSourceManagerEditor/Private/Widgets/SMediaSourceManagerChannel.cpp
@@ -26,6 +26,11 @@
 
 SMediaSourceManagerChannel::~SMediaSourceManagerChannel()
 {
+	if (ChannelPtr != nullptr)
+	{
+		ChannelPtr->OnInputPropertyChanged.RemoveAll(this);
+	}
+
 	DismissErrorNotification();
 }
 
@@ -65,6 +70,20 @@ void SMediaSourceManagerChannel::Construct(const FArguments& InArgs,
 						]
 				]
 
+			// Input warning icon.
+			+ SHorizontalBox::Slot()
+				.AutoWidth()
+				.HAlign(HAlign_Left)
+				.VAlign(VAlign_Center)
+				.Padding(4.0f, 0.0f, 0.0f, 0.0f)
+				[
+					SNew(SImage)
+						.Image(FCoreStyle::Get().GetBrush("Icons.Warning"))
+						.ToolTipText(LOCTEXT("InputWarning", "Input has incorrect settings."))
+						.Visibility(this,
+							&SMediaSourceManagerChannel::HandleInputWarningIconVisibility)
+				]
+
 			// Edit input.
 			+ SHorizontalBox::Slot()
 				.AutoWidth()
@@ -92,6 +111,7 @@ void SMediaSourceManagerChannel::Construct(const FArguments& InArgs,
 	// Start playing.
 	if (ChannelPtr != nullptr)
 	{
+		ChannelPtr->OnInputPropertyChanged.AddSP(this, &SMediaSourceManagerChannel::Refresh);
 		ChannelPtr->Play();
 	}
 }
@@ -307,7 +327,6 @@ void SMediaSourceManagerChannel::DismissErrorNotification()
 	}
 }
 
-
 FReply SMediaSourceManagerChannel::OnEditInput()
 {
 	// Get our input.
@@ -335,6 +354,11 @@ FReply SMediaSourceManagerChannel::OnEditInput()
 	return FReply::Handled();
 }
 
+EVisibility SMediaSourceManagerChannel::HandleInputWarningIconVisibility() const
+{
+	return bIsInputValid ? EVisibility::Hidden : EVisibility::Visible;
+}
+
 void SMediaSourceManagerChannel::Refresh()
 {
 	// Get channel.
@@ -345,13 +369,22 @@ void SMediaSourceManagerChannel::Refresh()
 
 		// Get input.
 		UMediaSourceManagerInput* Input = Channel->Input;
+		UMediaSource* MediaSource = nullptr;
 		if (Input != nullptr)
 		{
 			InputName = FText::FromString(Input->GetDisplayName());
+			MediaSource = Input->GetMediaSource();;
 		}
 
 		// Update input widgets.
 		InputNameTextBlock->SetText(InputName);
+
+		// Is the input valid?
+		bIsInputValid = true;
+		if (MediaSource != nullptr)
+		{
+			bIsInputValid = MediaSource->Validate();
+		}
 	}
 }
 
diff --git a/Engine/Plugins/Media/MediaSourceManager/Source/MediaSourceManagerEditor/Private/Widgets/SMediaSourceManagerChannel.h b/Engine/Plugins/Media/MediaSourceManager/Source/MediaSourceManagerEditor/Private/Widgets/SMediaSourceManagerChannel.h
--- a/Engine/Plugins/Media/MediaSourceManager/Source/MediaSourceManagerEditor/Private/Widgets/SMediaSourceManagerChannel.h
+++ b/Engine/Plugins/Media/MediaSourceManager/Source/MediaSourceManagerEditor/Private/Widgets/SMediaSourceManagerChannel.h
@@ -62,6 +62,9 @@ private:
 	 */
 	TSharedRef<SWidget> BuildMediaSourcePickerWidget();
 
+	/** Caches if the input is valid or not. */
+	bool bIsInputValid;
+
 	/**
 	 * Callback to add a media source.
 	 */
@@ -97,6 +100,11 @@ private:
 	 */
 	FReply OnEditInput();
 
+	/**
+	 * Determines warning icon visibilty based on the validity of the input.
+	 */
+	EVisibility HandleInputWarningIconVisibility() const;
+
 	/**
 	 * Refreshes the widgets based on the current sstate.
 	 */
