MediaSourceManager: Iterates through media io connections instead of devices for inputs as this is better for AJA devices.

#rnx

diff --git a/Engine/Plugins/Media/MediaSourceManager/Source/MediaSourceManagerEditor/MediaSourceManagerEditor.Build.cs b/Engine/Plugins/Media/MediaSourceManager/Source/MediaSourceManagerEditor/MediaSourceManagerEditor.Build.cs
--- a/Engine/Plugins/Media/MediaSourceManager/Source/MediaSourceManagerEditor/MediaSourceManagerEditor.Build.cs
+++ b/Engine/Plugins/Media/MediaSourceManager/Source/MediaSourceManagerEditor/MediaSourceManagerEditor.Build.cs
@@ -17,6 +17,7 @@ public class MediaSourceManagerEditor : ModuleRules
 				"InputCore",
 				"MediaAssets",
 				"MediaIOCore",
+				"MediaIOEditor",
 				"MediaPlayerEditor",
 				"MediaSourceManager",
 				"PropertyEditor",
diff --git a/Engine/Plugins/Media/MediaSourceManager/Source/MediaSourceManagerEditor/Private/Widgets/SMediaSourceManagerChannel.cpp b/Engine/Plugins/Media/MediaSourceManager/Source/MediaSourceManagerEditor/Private/Widgets/SMediaSourceManagerChannel.cpp
--- a/Engine/Plugins/Media/MediaSourceManager/Source/MediaSourceManagerEditor/Private/Widgets/SMediaSourceManagerChannel.cpp
+++ b/Engine/Plugins/Media/MediaSourceManager/Source/MediaSourceManagerEditor/Private/Widgets/SMediaSourceManagerChannel.cpp
@@ -12,6 +12,7 @@
 #include "IMediaIOCoreDeviceProvider.h"
 #include "IMediaIOCoreModule.h"
 #include "Inputs/MediaSourceManagerInputMediaSource.h"
+#include "MediaIOPermutationsSelectorBuilder.h"
 #include "MediaSource.h"
 #include "MediaSourceManagerChannel.h"
 #include "Subsystems/AssetEditorSubsystem.h"
@@ -176,16 +177,29 @@ TSharedRef<SWidget> SMediaSourceManagerChannel::CreateAssignInputMenu()
 			FName ProviderName = DeviceProvider->GetFName();
 			MenuBuilder.BeginSection(ProviderName, FText::FromName(ProviderName));
 
-			// Go over all devices.
-			TArray<FMediaIODevice> Devices = DeviceProvider->GetDevices();
-			for (const FMediaIODevice& Device : Devices)
+			// Go over all connections.
+			TArray<FMediaIOConnection> Connections = DeviceProvider->GetConnections();
+			for (const FMediaIOConnection& Connection : Connections)
 			{
-				// Add this device.
-				FName DeviceName = Device.DeviceName;
+				// Add this connection.
+				FText DeviceName = FText::FromName(Connection.Device.DeviceName);
+				FText LinkName = FMediaIOPermutationsSelectorBuilder::GetLabel(
+					FMediaIOPermutationsSelectorBuilder::NAME_TransportType,
+					Connection);
+				FText MenuText;
+				if (DeviceProvider->ShowInputTransportInSelector())
+				{
+					MenuText = FText::Format(LOCTEXT("Connection", "{0}: {1}"),
+						DeviceName, LinkName);
+				}
+				else
+				{
+					MenuText = DeviceName;
+				}
 
 				FUIAction AssignMediaIOInputAction(FExecuteAction::CreateSP(this,
-					&SMediaSourceManagerChannel::AssignMediaIOInput, DeviceProvider, Device));
-				MenuBuilder.AddMenuEntry(FText::FromName(DeviceName), FText(), FSlateIcon(),
+					&SMediaSourceManagerChannel::AssignMediaIOInput, DeviceProvider, Connection));
+				MenuBuilder.AddMenuEntry(MenuText, FText(), FSlateIcon(),
 					AssignMediaIOInputAction);
 			}
 
@@ -278,7 +292,7 @@ void SMediaSourceManagerChannel::AssignMediaSourceInput(UMediaSource* MediaSourc
 }
 
 void SMediaSourceManagerChannel::AssignMediaIOInput(IMediaIOCoreDeviceProvider* DeviceProvider,
-	FMediaIODevice Device)
+	FMediaIOConnection Connection)
 {
 	UMediaSourceManagerChannel* Channel = ChannelPtr.Get();
 	if (Channel != nullptr)
@@ -286,7 +300,7 @@ void SMediaSourceManagerChannel::AssignMediaIOInput(IMediaIOCoreDeviceProvider*
 		// Create media source.
 		FMediaIOConfiguration Configuration;
 		Configuration = DeviceProvider->GetDefaultConfiguration();
-		Configuration.MediaConnection.Device = Device;
+		Configuration.MediaConnection = Connection;
 		UMediaSource* MediaSource = DeviceProvider->CreateMediaSource(Configuration, Channel);
 		if (MediaSource != nullptr)
 		{
diff --git a/Engine/Plugins/Media/MediaSourceManager/Source/MediaSourceManagerEditor/Private/Widgets/SMediaSourceManagerChannel.h b/Engine/Plugins/Media/MediaSourceManager/Source/MediaSourceManagerEditor/Private/Widgets/SMediaSourceManagerChannel.h
--- a/Engine/Plugins/Media/MediaSourceManager/Source/MediaSourceManagerEditor/Private/Widgets/SMediaSourceManagerChannel.h
+++ b/Engine/Plugins/Media/MediaSourceManager/Source/MediaSourceManagerEditor/Private/Widgets/SMediaSourceManagerChannel.h
@@ -12,7 +12,7 @@ class UMediaSource;
 class UMediaSourceManagerChannel;
 
 struct FAssetData;
-struct FMediaIODevice;
+struct FMediaIOConnection;
 
 /**
  * Implements a single channel.
@@ -86,9 +86,9 @@ private:
 	 * Assigns a MediaIO input to this channel.
 	 *
 	 * @param DeviceProvider	MediaIO device provider to get the input from.
-	 * @param Device			Device to pass to the device provider.
+	 * @param Connection		Connection to pass to the device provider.
 	 */
-	void AssignMediaIOInput(IMediaIOCoreDeviceProvider* DeviceProvider, FMediaIODevice Device);
+	void AssignMediaIOInput(IMediaIOCoreDeviceProvider* DeviceProvider, FMediaIOConnection Connection);
 
 	/**
 	 * Call this to remove the error notification.
