MediaSourceManager: Added menu option to clear an input on a channel.
#rnx

diff --git a/Engine/Plugins/Media/MediaSourceManager/Source/MediaSourceManagerEditor/Private/Widgets/SMediaSourceManagerChannel.cpp b/Engine/Plugins/Media/MediaSourceManager/Source/MediaSourceManagerEditor/Private/Widgets/SMediaSourceManagerChannel.cpp
--- a/Engine/Plugins/Media/MediaSourceManager/Source/MediaSourceManagerEditor/Private/Widgets/SMediaSourceManagerChannel.cpp
+++ b/Engine/Plugins/Media/MediaSourceManager/Source/MediaSourceManagerEditor/Private/Widgets/SMediaSourceManagerChannel.cpp
@@ -13,6 +13,7 @@
 #include "IMediaIOCoreModule.h"
 #include "Inputs/MediaSourceManagerInputMediaSource.h"
 #include "MediaIOPermutationsSelectorBuilder.h"
+#include "MediaPlayer.h"
 #include "MediaSource.h"
 #include "MediaSourceManagerChannel.h"
 #include "Subsystems/AssetEditorSubsystem.h"
@@ -163,6 +164,14 @@ TSharedRef<SWidget> SMediaSourceManagerChannel::CreateAssignInputMenu()
 {
 	FMenuBuilder MenuBuilder(true, NULL);
 
+	// Add current asset options.
+	MenuBuilder.BeginSection(NAME_None, LOCTEXT("CurrentAsset", "Current Asset"));
+	FUIAction ClearAction(FExecuteAction::CreateSP(this, &SMediaSourceManagerChannel::ClearInput));
+	MenuBuilder.AddMenuEntry(LOCTEXT("Clear", "Clear"),
+		LOCTEXT("ClearToolTip", "Clears the asset set on this field"),
+		FSlateIcon(FAppStyle::GetAppStyleSetName(), "GenericCommands.Delete"), ClearAction);
+	MenuBuilder.EndSection();
+
 	// Get all Media IO device providers.
 	IMediaIOCoreModule& MediaIOCoreModule = IMediaIOCoreModule::Get();
 	TConstArrayView<IMediaIOCoreDeviceProvider*> DeviceProviders =
@@ -274,6 +283,24 @@ void SMediaSourceManagerChannel::AddMediaSourceEnterPressed(const TArray<FAssetD
 	}
 }
 
+void SMediaSourceManagerChannel::ClearInput()
+{
+	UMediaSourceManagerChannel* Channel = ChannelPtr.Get();
+	if (Channel != nullptr)
+	{
+		// Clear input on channel.
+		Channel->Modify();
+		Channel->Input = nullptr;
+
+		// Stop player.
+		UMediaPlayer* MediaPlayer = Channel->GetMediaPlayer();
+		if (MediaPlayer != nullptr)
+		{
+			MediaPlayer->Close();
+		}
+		Refresh();
+	}
+}
 
 void SMediaSourceManagerChannel::AssignMediaSourceInput(UMediaSource* MediaSource)
 {
diff --git a/Engine/Plugins/Media/MediaSourceManager/Source/MediaSourceManagerEditor/Private/Widgets/SMediaSourceManagerChannel.h b/Engine/Plugins/Media/MediaSourceManager/Source/MediaSourceManagerEditor/Private/Widgets/SMediaSourceManagerChannel.h
--- a/Engine/Plugins/Media/MediaSourceManager/Source/MediaSourceManagerEditor/Private/Widgets/SMediaSourceManagerChannel.h
+++ b/Engine/Plugins/Media/MediaSourceManager/Source/MediaSourceManagerEditor/Private/Widgets/SMediaSourceManagerChannel.h
@@ -76,7 +76,12 @@ private:
 	void AddMediaSourceEnterPressed(const TArray<FAssetData>& AssetData);
 
 	/**
-	 * Assignss a media source to this channel.
+	 * Clears the input on this channel.
+	 */
+	void ClearInput();
+
+	/**
+	 * Assigns a media source to this channel.
 	 * 
 	 * @param MediaSource	Media source to use as input.
 	 */
