MediaSourceManager: Channel input edit functionality moved to assign input menu.
#rnx

diff --git a/Engine/Plugins/Media/MediaSourceManager/Source/MediaSourceManagerEditor/Private/Widgets/SMediaSourceManagerChannel.cpp b/Engine/Plugins/Media/MediaSourceManager/Source/MediaSourceManagerEditor/Private/Widgets/SMediaSourceManagerChannel.cpp
--- a/Engine/Plugins/Media/MediaSourceManager/Source/MediaSourceManagerEditor/Private/Widgets/SMediaSourceManagerChannel.cpp
+++ b/Engine/Plugins/Media/MediaSourceManager/Source/MediaSourceManagerEditor/Private/Widgets/SMediaSourceManagerChannel.cpp
@@ -16,6 +16,7 @@
 #include "MediaPlayer.h"
 #include "MediaSource.h"
 #include "MediaSourceManagerChannel.h"
+#include "Styling/AppStyle.h"
 #include "Subsystems/AssetEditorSubsystem.h"
 #include "Widgets/Input/SButton.h"
 #include "Widgets/Input/SComboButton.h"
@@ -86,18 +87,6 @@ void SMediaSourceManagerChannel::Construct(const FArguments& InArgs,
 							&SMediaSourceManagerChannel::HandleInputWarningIconVisibility)
 				]
 
-			// Edit input.
-			+ SHorizontalBox::Slot()
-				.AutoWidth()
-				.HAlign(HAlign_Left)
-				.Padding(2)
-				[
-					SNew(SButton)
-						.ContentPadding(3)
-						.OnClicked(this, &SMediaSourceManagerChannel::OnEditInput)
-						.Text(LOCTEXT("EditInput", "Edit Input"))
-				]
-
 			// Out texture
 			+ SHorizontalBox::Slot()
 				.FillWidth(0.11f)
@@ -165,12 +154,25 @@ TSharedRef<SWidget> SMediaSourceManagerChannel::CreateAssignInputMenu()
 	FMenuBuilder MenuBuilder(true, NULL);
 
 	// Add current asset options.
-	MenuBuilder.BeginSection(NAME_None, LOCTEXT("CurrentAsset", "Current Asset"));
-	FUIAction ClearAction(FExecuteAction::CreateSP(this, &SMediaSourceManagerChannel::ClearInput));
-	MenuBuilder.AddMenuEntry(LOCTEXT("Clear", "Clear"),
-		LOCTEXT("ClearToolTip", "Clears the asset set on this field"),
-		FSlateIcon(FAppStyle::GetAppStyleSetName(), "GenericCommands.Delete"), ClearAction);
-	MenuBuilder.EndSection();
+	UMediaSourceManagerChannel* Channel = ChannelPtr.Get();
+	if ((Channel != nullptr) && (Channel->Input != nullptr))
+	{
+		MenuBuilder.BeginSection(NAME_None, LOCTEXT("CurrentAsset", "Current Asset"));
+
+		// Edit.
+		FUIAction EditAction(FExecuteAction::CreateSP(this, &SMediaSourceManagerChannel::OnEditInput));
+		MenuBuilder.AddMenuEntry(LOCTEXT("Edit", "Edit"),
+			LOCTEXT("EditToolTip", "Edit this asset"),
+			FSlateIcon(FAppStyle::GetAppStyleSetName(), "Icons.Edit"), EditAction);
+
+		// Clear.
+		FUIAction ClearAction(FExecuteAction::CreateSP(this, &SMediaSourceManagerChannel::ClearInput));
+		MenuBuilder.AddMenuEntry(LOCTEXT("Clear", "Clear"),
+			LOCTEXT("ClearToolTip", "Clears the asset set on this field"),
+			FSlateIcon(FAppStyle::GetAppStyleSetName(), "GenericCommands.Delete"), ClearAction);
+
+		MenuBuilder.EndSection();
+	}
 
 	// Get all Media IO device providers.
 	IMediaIOCoreModule& MediaIOCoreModule = IMediaIOCoreModule::Get();
@@ -368,7 +370,7 @@ void SMediaSourceManagerChannel::DismissErrorNotification()
 	}
 }
 
-FReply SMediaSourceManagerChannel::OnEditInput()
+void SMediaSourceManagerChannel::OnEditInput()
 {
 	// Get our input.
 	TArray<UObject*> AssetArray;
@@ -391,8 +393,6 @@ FReply SMediaSourceManagerChannel::OnEditInput()
 	{
 		GEditor->GetEditorSubsystem<UAssetEditorSubsystem>()->OpenEditorForAssets(AssetArray);
 	}
-
-	return FReply::Handled();
 }
 
 EVisibility SMediaSourceManagerChannel::HandleInputWarningIconVisibility() const
diff --git a/Engine/Plugins/Media/MediaSourceManager/Source/MediaSourceManagerEditor/Private/Widgets/SMediaSourceManagerChannel.h b/Engine/Plugins/Media/MediaSourceManager/Source/MediaSourceManagerEditor/Private/Widgets/SMediaSourceManagerChannel.h
--- a/Engine/Plugins/Media/MediaSourceManager/Source/MediaSourceManagerEditor/Private/Widgets/SMediaSourceManagerChannel.h
+++ b/Engine/Plugins/Media/MediaSourceManager/Source/MediaSourceManagerEditor/Private/Widgets/SMediaSourceManagerChannel.h
@@ -101,9 +101,9 @@ private:
 	void DismissErrorNotification();
 
 	/**
-	 * Called when the edit input button is pressed.
+	 * Call this to bring up an editor to edit the input.
 	 */
-	FReply OnEditInput();
+	void OnEditInput();
 
 	/**
 	 * Determines warning icon visibilty based on the validity of the input.
