MediaSourceManager: Double clicking on the manager asset now brings up the tab editor and not the asset editor so it can be docked.
#rnx

diff --git a/Engine/Plugins/Media/MediaSourceManager/Source/MediaSourceManagerEditor/Private/AssetTools/MediaSourceManagerActions.cpp b/Engine/Plugins/Media/MediaSourceManager/Source/MediaSourceManagerEditor/Private/AssetTools/MediaSourceManagerActions.cpp
--- a/Engine/Plugins/Media/MediaSourceManager/Source/MediaSourceManagerEditor/Private/AssetTools/MediaSourceManagerActions.cpp
+++ b/Engine/Plugins/Media/MediaSourceManager/Source/MediaSourceManagerEditor/Private/AssetTools/MediaSourceManagerActions.cpp
@@ -4,6 +4,8 @@
 
 #include "AssetRegistry/AssetData.h"
 #include "MediaSourceManager.h"
+#include "MediaSourceManagerEditorModule.h"
+#include "Modules/ModuleManager.h"
 #include "Toolkits/MediaSourceManagerEditorToolkit.h"
 
 #define LOCTEXT_NAMESPACE "MediaSourceManagerActions"
@@ -61,8 +63,12 @@ void FMediaSourceManagerActions::OpenAssetEditor(const TArray<UObject*>& InObjec
 
 		if (MediaSourceManager != nullptr)
 		{
-			TSharedRef<FMediaSourceManagerEditorToolkit> EditorToolkit = MakeShareable(new FMediaSourceManagerEditorToolkit(Style));
-			EditorToolkit->Initialize(MediaSourceManager, Mode, EditWithinLevelEditor);
+			IMediaSourceManagerEditorModule* Module = FModuleManager::GetModulePtr
+				<IMediaSourceManagerEditorModule>("MediaSourceManagerEditor");
+			if (Module != nullptr)
+			{
+				Module->OpenEditorTab();
+			}
 		}
 	}
 }
diff --git a/Engine/Plugins/Media/MediaSourceManager/Source/MediaSourceManagerEditor/Private/MediaSourceManagerEditorModule.cpp b/Engine/Plugins/Media/MediaSourceManager/Source/MediaSourceManagerEditor/Private/MediaSourceManagerEditorModule.cpp
--- a/Engine/Plugins/Media/MediaSourceManager/Source/MediaSourceManagerEditor/Private/MediaSourceManagerEditorModule.cpp
+++ b/Engine/Plugins/Media/MediaSourceManager/Source/MediaSourceManagerEditor/Private/MediaSourceManagerEditorModule.cpp
@@ -27,6 +27,15 @@ public:
 	//~ IMediaSourceManagerEditorModule interface
 	TSharedPtr<ISlateStyle> GetStyle() { return Style; }
 
+	void OpenEditorTab()
+	{
+		FTabId ManagerTabID(ManagerTabName);
+		FGlobalTabmanager::Get()->TryInvokeTab(ManagerTabID);
+
+		FTabId PreviewTabID(PreviewTabName);
+		FGlobalTabmanager::Get()->TryInvokeTab(PreviewTabID);
+	}
+
 	//~ IModuleInterface interface
 
 	virtual void StartupModule() override
diff --git a/Engine/Plugins/Media/MediaSourceManager/Source/MediaSourceManagerEditor/Public/MediaSourceManagerEditorModule.h b/Engine/Plugins/Media/MediaSourceManager/Source/MediaSourceManagerEditor/Public/MediaSourceManagerEditorModule.h
--- a/Engine/Plugins/Media/MediaSourceManager/Source/MediaSourceManagerEditor/Public/MediaSourceManagerEditorModule.h
+++ b/Engine/Plugins/Media/MediaSourceManager/Source/MediaSourceManagerEditor/Public/MediaSourceManagerEditorModule.h
@@ -24,4 +24,8 @@ public:
 	 **/
 	virtual TSharedPtr<ISlateStyle> GetStyle() = 0;
 
+	/**
+	 * Opens the editor tab to edit the manager.
+	 */
+	virtual void OpenEditorTab() = 0;
 };
