MediaSourceManager: Added SMediaSourceManagerPreview to handle previews for each channel.
#rnx

diff --git a/Engine/Plugins/Media/MediaSourceManager/Source/MediaSourceManagerEditor/Private/Widgets/SMediaSourceManagerPreview.cpp b/Engine/Plugins/Media/MediaSourceManager/Source/MediaSourceManagerEditor/Private/Widgets/SMediaSourceManagerPreview.cpp
new file mode 100644
--- /dev/null
+++ b/Engine/Plugins/Media/MediaSourceManager/Source/MediaSourceManagerEditor/Private/Widgets/SMediaSourceManagerPreview.cpp
@@ -0,0 +1,19 @@
+// Copyright Epic Games, Inc. All Rights Reserved.
+
+#include "Widgets/SMediaSourceManagerPreview.h"
+
+#include "Widgets/SMediaPlayerEditorViewer.h"
+
+#define LOCTEXT_NAMESPACE "SMediaSourceManagerPreview"
+
+void SMediaSourceManagerPreview::Construct(const FArguments& InArg, UMediaPlayer& InMediaPlayer,
+	UMediaTexture* InMediaTexture, const TSharedRef<ISlateStyle>& InStyle)
+{
+	ChildSlot
+		[
+			SNew(SMediaPlayerEditorViewer, InMediaPlayer, InMediaTexture, InStyle, false)
+				.bShowUrl(false)
+		];
+}
+
+#undef LOCTEXT_NAMESPACE
diff --git a/Engine/Plugins/Media/MediaSourceManager/Source/MediaSourceManagerEditor/Private/Widgets/SMediaSourceManagerPreview.h b/Engine/Plugins/Media/MediaSourceManager/Source/MediaSourceManagerEditor/Private/Widgets/SMediaSourceManagerPreview.h
new file mode 100644
--- /dev/null
+++ b/Engine/Plugins/Media/MediaSourceManager/Source/MediaSourceManagerEditor/Private/Widgets/SMediaSourceManagerPreview.h
@@ -0,0 +1,34 @@
+// Copyright Epic Games, Inc. All Rights Reserved.
+
+#pragma once
+
+#include "Widgets/DeclarativeSyntaxSupport.h"
+#include "Widgets/SCompoundWidget.h"
+
+class ISlateStyle;
+class UMediaPlayer;
+class UMediaTexture;
+
+/**
+ * Implements a preview for a single media source manager channel.
+ */
+class SMediaSourceManagerPreview
+	: public SCompoundWidget
+{
+public:
+
+	SLATE_BEGIN_ARGS(SMediaSourceManagerPreview) { }
+	SLATE_END_ARGS()
+
+	/**
+	 * Construct this widget.
+	 *
+	 * @param InArgs					The declaration data for this widget.
+	 * @param InMediaPlayer				The UMediaPlayer asset to show the details for.
+	 * @param InMediaTexture			The UMediaTexture asset to output video to. If nullptr then use our own.
+	 * @param InStyle					The style set to use.
+	 */
+	void Construct(const FArguments& InArgs, UMediaPlayer& InMediaPlayer,
+		UMediaTexture* InMediaTexture, const TSharedRef<ISlateStyle>& InStyle);
+
+};
diff --git a/Engine/Plugins/Media/MediaSourceManager/Source/MediaSourceManagerEditor/Private/Widgets/SMediaSourceManagerPreviews.cpp b/Engine/Plugins/Media/MediaSourceManager/Source/MediaSourceManagerEditor/Private/Widgets/SMediaSourceManagerPreviews.cpp
--- a/Engine/Plugins/Media/MediaSourceManager/Source/MediaSourceManagerEditor/Private/Widgets/SMediaSourceManagerPreviews.cpp
+++ b/Engine/Plugins/Media/MediaSourceManager/Source/MediaSourceManagerEditor/Private/Widgets/SMediaSourceManagerPreviews.cpp
@@ -8,7 +8,7 @@
 #include "MediaTexture.h"
 #include "Modules/ModuleManager.h"
 #include "Widgets/Layout/SScrollBox.h"
-#include "Widgets/SMediaPlayerEditorViewer.h"
+#include "Widgets/SMediaSourceManagerPreview.h"
 
 #define LOCTEXT_NAMESPACE "SMediaSourceManagerPreviews"
 
@@ -70,8 +70,7 @@ void SMediaSourceManagerPreviews::RefreshChannels()
 					.Padding(2)
 					.HAlign(HAlign_Center)
 					[
-						SNew(SMediaPlayerEditorViewer, *MediaPlayer, MediaTexture, StyleRef, false)
-							.bShowUrl(false)
+						SNew(SMediaSourceManagerPreview, *MediaPlayer, MediaTexture, StyleRef)
 					];
 
 				Channel->Play();
