MediaSourceManager: Copied over the ability to "drop media sources on a channel" to the preview.
#rnx

diff --git a/Engine/Plugins/Media/MediaSourceManager/Source/MediaSourceManagerEditor/Private/Widgets/SMediaSourceManagerPreview.cpp b/Engine/Plugins/Media/MediaSourceManager/Source/MediaSourceManagerEditor/Private/Widgets/SMediaSourceManagerPreview.cpp
--- a/Engine/Plugins/Media/MediaSourceManager/Source/MediaSourceManagerEditor/Private/Widgets/SMediaSourceManagerPreview.cpp
+++ b/Engine/Plugins/Media/MediaSourceManager/Source/MediaSourceManagerEditor/Private/Widgets/SMediaSourceManagerPreview.cpp
@@ -2,18 +2,86 @@
 
 #include "Widgets/SMediaSourceManagerPreview.h"
 
+#include "DragAndDrop/AssetDragDropOp.h"
+#include "Inputs/MediaSourceManagerInputMediaSource.h"
+#include "MediaSourceManagerChannel.h"
+#include "MediaSource.h"
+#include "MediaTexture.h"
 #include "Widgets/SMediaPlayerEditorViewer.h"
 
 #define LOCTEXT_NAMESPACE "SMediaSourceManagerPreview"
 
-void SMediaSourceManagerPreview::Construct(const FArguments& InArg, UMediaPlayer& InMediaPlayer,
-	UMediaTexture* InMediaTexture, const TSharedRef<ISlateStyle>& InStyle)
+void SMediaSourceManagerPreview::Construct(const FArguments& InArg,
+	UMediaSourceManagerChannel* InChannel, const TSharedRef<ISlateStyle>& InStyle)
 {
+	ChannelPtr = InChannel;
+
+	UMediaPlayer* MediaPlayer = InChannel->GetMediaPlayer();
+	UMediaTexture* MediaTexture = Cast<UMediaTexture>(InChannel->OutTexture);
+
 	ChildSlot
 		[
-			SNew(SMediaPlayerEditorViewer, InMediaPlayer, InMediaTexture, InStyle, false)
+			SNew(SMediaPlayerEditorViewer, *MediaPlayer, MediaTexture, InStyle, false)
 				.bShowUrl(false)
 		];
 }
 
+
+void SMediaSourceManagerPreview::OnDragEnter(const FGeometry& MyGeometry, const FDragDropEvent& DragDropEvent)
+{
+}
+
+void SMediaSourceManagerPreview::OnDragLeave(const FDragDropEvent& DragDropEvent)
+{
+}
+
+FReply SMediaSourceManagerPreview::OnDragOver(const FGeometry& MyGeometry, const FDragDropEvent& DragDropEvent)
+{
+	// Is this an asset drop?
+	if (TSharedPtr<FAssetDragDropOp> AssetDragDrop = DragDropEvent.GetOperationAs<FAssetDragDropOp>())
+	{
+		return FReply::Handled();
+	}
+
+	return FReply::Unhandled();
+}
+
+FReply SMediaSourceManagerPreview::OnDrop(const FGeometry& MyGeometry, const FDragDropEvent& DragDropEvent)
+{
+	// Is this an asset drop?
+	if (TSharedPtr<FAssetDragDropOp> AssetDragDrop = DragDropEvent.GetOperationAs<FAssetDragDropOp>())
+	{
+		for (const FAssetData& Asset : AssetDragDrop->GetAssets())
+		{
+			// Is this a media source?
+			UMediaSource* MediaSource = Cast<UMediaSource>(Asset.GetAsset());
+			if (MediaSource != nullptr)
+			{
+				AssignMediaSourceInput(MediaSource);
+				break;
+			}
+		}
+
+		return FReply::Handled();
+	}
+
+	return FReply::Unhandled();
+}
+
+
+void SMediaSourceManagerPreview::AssignMediaSourceInput(UMediaSource* MediaSource)
+{
+	UMediaSourceManagerChannel* Channel = ChannelPtr.Get();
+	if (Channel != nullptr)
+	{
+		// Assign to channel.
+		Channel->Modify();
+		UMediaSourceManagerInputMediaSource* Input = NewObject<UMediaSourceManagerInputMediaSource>(Channel);
+		Input->MediaSource = MediaSource;
+		Channel->Input = Input;
+		Channel->Play();
+	}
+}
+
+
 #undef LOCTEXT_NAMESPACE
diff --git a/Engine/Plugins/Media/MediaSourceManager/Source/MediaSourceManagerEditor/Private/Widgets/SMediaSourceManagerPreview.h b/Engine/Plugins/Media/MediaSourceManager/Source/MediaSourceManagerEditor/Private/Widgets/SMediaSourceManagerPreview.h
--- a/Engine/Plugins/Media/MediaSourceManager/Source/MediaSourceManagerEditor/Private/Widgets/SMediaSourceManagerPreview.h
+++ b/Engine/Plugins/Media/MediaSourceManager/Source/MediaSourceManagerEditor/Private/Widgets/SMediaSourceManagerPreview.h
@@ -6,8 +6,8 @@
 #include "Widgets/SCompoundWidget.h"
 
 class ISlateStyle;
-class UMediaPlayer;
-class UMediaTexture;
+class UMediaSource;
+class UMediaSourceManagerChannel;
 
 /**
  * Implements a preview for a single media source manager channel.
@@ -24,11 +24,29 @@ public:
 	 * Construct this widget.
 	 *
 	 * @param InArgs					The declaration data for this widget.
-	 * @param InMediaPlayer				The UMediaPlayer asset to show the details for.
-	 * @param InMediaTexture			The UMediaTexture asset to output video to. If nullptr then use our own.
+	 * @param InChannel					The channel to preview.
 	 * @param InStyle					The style set to use.
 	 */
-	void Construct(const FArguments& InArgs, UMediaPlayer& InMediaPlayer,
-		UMediaTexture* InMediaTexture, const TSharedRef<ISlateStyle>& InStyle);
+	void Construct(const FArguments& InArgs, UMediaSourceManagerChannel* InChannel,
+		const TSharedRef<ISlateStyle>& InStyle);
+
+	//~ SWidget interface
+	virtual void OnDragEnter(const FGeometry& MyGeometrsy, const FDragDropEvent& DragDropEvent) override;
+	virtual void OnDragLeave(const FDragDropEvent& DragDropEvent) override;
+	virtual FReply OnDragOver(const FGeometry& MyGeometry, const FDragDropEvent& DragDropEvent) override;
+	virtual FReply OnDrop(const FGeometry& MyGeometry, const FDragDropEvent& DragDropEvent) override;
+	//~ End of SWidget interface
+
+private:
+
+	/**
+	 * Assigns a media source to this channel.
+	 *
+	 * @param MediaSource	Media source to use as input.
+	 */
+	void AssignMediaSourceInput(UMediaSource* MediaSource);
+
+	/** Pointer to the object that is being viewed. */
+	TWeakObjectPtr<UMediaSourceManagerChannel> ChannelPtr;
 
 };
diff --git a/Engine/Plugins/Media/MediaSourceManager/Source/MediaSourceManagerEditor/Private/Widgets/SMediaSourceManagerPreviews.cpp b/Engine/Plugins/Media/MediaSourceManager/Source/MediaSourceManagerEditor/Private/Widgets/SMediaSourceManagerPreviews.cpp
--- a/Engine/Plugins/Media/MediaSourceManager/Source/MediaSourceManagerEditor/Private/Widgets/SMediaSourceManagerPreviews.cpp
+++ b/Engine/Plugins/Media/MediaSourceManager/Source/MediaSourceManagerEditor/Private/Widgets/SMediaSourceManagerPreviews.cpp
@@ -5,7 +5,6 @@
 #include "MediaSourceManager.h"
 #include "MediaSourceManagerEditorModule.h"
 #include "MediaSourceManagerSettings.h"
-#include "MediaTexture.h"
 #include "Modules/ModuleManager.h"
 #include "Widgets/Layout/SScrollBox.h"
 #include "Widgets/SMediaSourceManagerPreview.h"
@@ -62,15 +61,12 @@ void SMediaSourceManagerPreviews::RefreshChannels()
 			UMediaSourceManagerChannel* Channel = MediaSourceManager->Channels[Index];
 			if (Channel != nullptr)
 			{
-				UMediaPlayer* MediaPlayer = Channel->GetMediaPlayer();
-				UMediaTexture* MediaTexture = Cast<UMediaTexture>(Channel->OutTexture);
-
 				// Add preview widget.
 				ChannelsContainer->AddSlot()
 					.Padding(2)
 					.HAlign(HAlign_Center)
 					[
-						SNew(SMediaSourceManagerPreview, *MediaPlayer, MediaTexture, StyleRef)
+						SNew(SMediaSourceManagerPreview, Channel, StyleRef)
 					];
 
 				Channel->Play();
