MediaSourceManager: Copied over ability to "drag from the thumbnail to an actor to apply our material" to the preview.
#rnx

diff --git a/Engine/Plugins/Media/MediaSourceManager/Source/MediaSourceManagerEditor/Private/Widgets/SMediaSourceManagerPreview.cpp b/Engine/Plugins/Media/MediaSourceManager/Source/MediaSourceManagerEditor/Private/Widgets/SMediaSourceManagerPreview.cpp
--- a/Engine/Plugins/Media/MediaSourceManager/Source/MediaSourceManagerEditor/Private/Widgets/SMediaSourceManagerPreview.cpp
+++ b/Engine/Plugins/Media/MediaSourceManager/Source/MediaSourceManagerEditor/Private/Widgets/SMediaSourceManagerPreview.cpp
@@ -4,6 +4,7 @@
 
 #include "DragAndDrop/AssetDragDropOp.h"
 #include "Inputs/MediaSourceManagerInputMediaSource.h"
+#include "Materials/MaterialInstanceConstant.h"
 #include "MediaSourceManagerChannel.h"
 #include "MediaSource.h"
 #include "MediaTexture.h"
@@ -11,6 +12,8 @@
 
 #define LOCTEXT_NAMESPACE "SMediaSourceManagerPreview"
 
+FLazyName SMediaSourceManagerPreview::MediaTextureName("MediaTexture");
+
 void SMediaSourceManagerPreview::Construct(const FArguments& InArg,
 	UMediaSourceManagerChannel* InChannel, const TSharedRef<ISlateStyle>& InStyle)
 {
@@ -19,11 +22,17 @@ void SMediaSourceManagerPreview::Construct(const FArguments& InArg,
 	UMediaPlayer* MediaPlayer = InChannel->GetMediaPlayer();
 	UMediaTexture* MediaTexture = Cast<UMediaTexture>(InChannel->OutTexture);
 
+	TSharedPtr<SMediaPlayerEditorViewer> PlayerViewer;
 	ChildSlot
 		[
-			SNew(SMediaPlayerEditorViewer, *MediaPlayer, MediaTexture, InStyle, false)
+			SAssignNew(PlayerViewer, SMediaPlayerEditorViewer, *MediaPlayer, MediaTexture, InStyle, false)
 				.bShowUrl(false)
 		];
+
+	if (PlayerViewer.IsValid())
+	{
+		PlayerViewer->EnableMouseControl(false);
+	}
 }
 
 
@@ -69,6 +78,37 @@ FReply SMediaSourceManagerPreview::OnDrop(const FGeometry& MyGeometry, const FDr
 }
 
 
+FReply SMediaSourceManagerPreview::OnMouseButtonDown(const FGeometry& MyGeometry, const FPointerEvent& MouseEvent)
+{
+	if (MouseEvent.GetEffectingButton() == EKeys::LeftMouseButton)
+	{
+		return FReply::Handled().DetectDrag(SharedThis(this), EKeys::LeftMouseButton);
+	}
+
+	return FReply::Unhandled();
+}
+
+FReply SMediaSourceManagerPreview::OnDragDetected(const FGeometry& MyGeometry, const FPointerEvent& MouseEvent)
+{
+	UMediaSourceManagerChannel* Channel = ChannelPtr.Get();
+	if (Channel != nullptr)
+	{
+		UTexture* Texture = Channel->OutTexture;
+		if (Texture != nullptr)
+		{
+			UMaterialInstanceConstant* Material = GetMaterial();
+			if (Material != nullptr)
+			{
+
+				FAssetData AssetData(Material);
+				return FReply::Handled().BeginDragDrop(FAssetDragDropOp::New(AssetData));
+			}
+		}
+	}
+
+	return FReply::Unhandled();
+}
+
 void SMediaSourceManagerPreview::AssignMediaSourceInput(UMediaSource* MediaSource)
 {
 	UMediaSourceManagerChannel* Channel = ChannelPtr.Get();
@@ -84,4 +124,48 @@ void SMediaSourceManagerPreview::AssignMediaSourceInput(UMediaSource* MediaSourc
 }
 
 
+UMaterialInstanceConstant* SMediaSourceManagerPreview::GetMaterial()
+{
+	UMaterialInstanceConstant* MaterialInstance = nullptr;
+
+	UMediaSourceManagerChannel* Channel = ChannelPtr.Get();
+	if (Channel != nullptr)
+	{
+		// Do we already have a material?
+		MaterialInstance = Channel->Material;
+		if (MaterialInstance == nullptr)
+		{
+			// No. Create one.
+			UMaterial* Material = LoadObject<UMaterial>(NULL,
+				TEXT("/MediaSourceManager/M_MediaSourceManager"), NULL, LOAD_None, NULL);
+			if (Material != nullptr)
+			{
+				UObject* Outer = Channel->GetOuter();
+				Channel->Modify();
+
+				FString MaterialName = Material->GetName();
+				if (MaterialName.StartsWith(TEXT("M_")))
+				{
+					MaterialName.InsertAt(1, TEXT("I"));
+				}
+				FName MaterialUniqueName = MakeUniqueObjectName(Outer, UMaterialInstanceConstant::StaticClass(),
+					FName(*MaterialName));
+
+				// Create instance.
+				MaterialInstance =
+					NewObject<UMaterialInstanceConstant>(Outer, MaterialUniqueName, RF_Public);
+				Channel->Material = MaterialInstance;
+				MaterialInstance->SetParentEditorOnly(Material);
+				MaterialInstance->SetTextureParameterValueEditorOnly(
+					FMaterialParameterInfo(MediaTextureName),
+					Channel->OutTexture);
+				MaterialInstance->PostEditChange();
+			}
+		}
+	}
+
+	return MaterialInstance;
+}
+
+
 #undef LOCTEXT_NAMESPACE
diff --git a/Engine/Plugins/Media/MediaSourceManager/Source/MediaSourceManagerEditor/Private/Widgets/SMediaSourceManagerPreview.h b/Engine/Plugins/Media/MediaSourceManager/Source/MediaSourceManagerEditor/Private/Widgets/SMediaSourceManagerPreview.h
--- a/Engine/Plugins/Media/MediaSourceManager/Source/MediaSourceManagerEditor/Private/Widgets/SMediaSourceManagerPreview.h
+++ b/Engine/Plugins/Media/MediaSourceManager/Source/MediaSourceManagerEditor/Private/Widgets/SMediaSourceManagerPreview.h
@@ -6,6 +6,7 @@
 #include "Widgets/SCompoundWidget.h"
 
 class ISlateStyle;
+class UMaterialInstanceConstant;
 class UMediaSource;
 class UMediaSourceManagerChannel;
 
@@ -35,6 +36,9 @@ public:
 	virtual void OnDragLeave(const FDragDropEvent& DragDropEvent) override;
 	virtual FReply OnDragOver(const FGeometry& MyGeometry, const FDragDropEvent& DragDropEvent) override;
 	virtual FReply OnDrop(const FGeometry& MyGeometry, const FDragDropEvent& DragDropEvent) override;
+	
+	virtual FReply OnMouseButtonDown(const FGeometry& MyGeometry, const FPointerEvent& MouseEvent) override;
+	virtual FReply OnDragDetected(const FGeometry& MyGeometry, const FPointerEvent& MouseEvent) override;
 	//~ End of SWidget interface
 
 private:
@@ -46,7 +50,15 @@ private:
 	 */
 	void AssignMediaSourceInput(UMediaSource* MediaSource);
 
+	/**
+	 * Gets a material for this texture.
+	 */
+	UMaterialInstanceConstant* GetMaterial();
+
 	/** Pointer to the object that is being viewed. */
 	TWeakObjectPtr<UMediaSourceManagerChannel> ChannelPtr;
 
+	/** Name for the media texture parameter in the material. */
+	static FLazyName MediaTextureName;
+
 };
