MediaSourceManager: Copied context menu edit/clear entries to the preview widget.
#rnx

diff --git a/Engine/Plugins/Media/MediaSourceManager/Source/MediaSourceManagerEditor/Private/Widgets/SMediaSourceManagerPreview.cpp b/Engine/Plugins/Media/MediaSourceManager/Source/MediaSourceManagerEditor/Private/Widgets/SMediaSourceManagerPreview.cpp
--- a/Engine/Plugins/Media/MediaSourceManager/Source/MediaSourceManagerEditor/Private/Widgets/SMediaSourceManagerPreview.cpp
+++ b/Engine/Plugins/Media/MediaSourceManager/Source/MediaSourceManagerEditor/Private/Widgets/SMediaSourceManagerPreview.cpp
@@ -3,11 +3,16 @@
 #include "Widgets/SMediaSourceManagerPreview.h"
 
 #include "DragAndDrop/AssetDragDropOp.h"
+#include "Editor.h"
+#include "Framework/Application/SlateApplication.h"
+#include "Framework/MultiBox/MultiBoxBuilder.h"
 #include "Inputs/MediaSourceManagerInputMediaSource.h"
 #include "Materials/MaterialInstanceConstant.h"
+#include "MediaPlayer.h"
 #include "MediaSourceManagerChannel.h"
 #include "MediaSource.h"
 #include "MediaTexture.h"
+#include "Subsystems/AssetEditorSubsystem.h"
 #include "Widgets/SMediaPlayerEditorViewer.h"
 
 #define LOCTEXT_NAMESPACE "SMediaSourceManagerPreview"
@@ -88,6 +93,18 @@ FReply SMediaSourceManagerPreview::OnMouseButtonDown(const FGeometry& MyGeometry
 	return FReply::Unhandled();
 }
 
+FReply SMediaSourceManagerPreview::OnMouseButtonUp(const FGeometry& MyGeometry, const FPointerEvent& MouseEvent)
+{
+	if (MouseEvent.GetEffectingButton() == EKeys::RightMouseButton)
+	{
+		OpenContextMenu(MouseEvent);
+
+		return FReply::Handled();
+	}
+
+	return FReply::Unhandled();
+}
+
 FReply SMediaSourceManagerPreview::OnDragDetected(const FGeometry& MyGeometry, const FPointerEvent& MouseEvent)
 {
 	UMediaSourceManagerChannel* Channel = ChannelPtr.Get();
@@ -109,6 +126,24 @@ FReply SMediaSourceManagerPreview::OnDragDetected(const FGeometry& MyGeometry, c
 	return FReply::Unhandled();
 }
 
+void SMediaSourceManagerPreview::ClearInput()
+{
+	UMediaSourceManagerChannel* Channel = ChannelPtr.Get();
+	if (Channel != nullptr)
+	{
+		// Clear input on channel.
+		Channel->Modify();
+		Channel->Input = nullptr;
+
+		// Stop player.
+		UMediaPlayer* MediaPlayer = Channel->GetMediaPlayer();
+		if (MediaPlayer != nullptr)
+		{
+			MediaPlayer->Close();
+		}
+	}
+}
+
 void SMediaSourceManagerPreview::AssignMediaSourceInput(UMediaSource* MediaSource)
 {
 	UMediaSourceManagerChannel* Channel = ChannelPtr.Get();
@@ -123,6 +158,30 @@ void SMediaSourceManagerPreview::AssignMediaSourceInput(UMediaSource* MediaSourc
 	}
 }
 
+void SMediaSourceManagerPreview::OnEditInput()
+{
+	// Get our input.
+	TArray<UObject*> AssetArray;
+	UMediaSourceManagerChannel* Channel = ChannelPtr.Get();
+	if (Channel != nullptr)
+	{
+		UMediaSourceManagerInput* Input = Channel->Input;
+		if (Input != nullptr)
+		{
+			UMediaSource* MediaSource = Input->GetMediaSource();
+			if (MediaSource != nullptr)
+			{
+				AssetArray.Add(MediaSource);
+			}
+		}
+	}
+
+	// Open the editor.
+	if (AssetArray.Num() > 0)
+	{
+		GEditor->GetEditorSubsystem<UAssetEditorSubsystem>()->OpenEditorForAssets(AssetArray);
+	}
+}
 
 UMaterialInstanceConstant* SMediaSourceManagerPreview::GetMaterial()
 {
@@ -167,5 +226,37 @@ UMaterialInstanceConstant* SMediaSourceManagerPreview::GetMaterial()
 	return MaterialInstance;
 }
 
+void SMediaSourceManagerPreview::OpenContextMenu(const FPointerEvent& MouseEvent)
+{
+	const bool bInShouldCloseWindowAfterMenuSelection = true;
+	FMenuBuilder MenuBuilder(bInShouldCloseWindowAfterMenuSelection, nullptr);
+
+	// Add current asset options.
+	UMediaSourceManagerChannel* Channel = ChannelPtr.Get();
+	if ((Channel != nullptr) && (Channel->Input != nullptr))
+	{
+		MenuBuilder.BeginSection(NAME_None, LOCTEXT("CurrentAsset", "Current Asset"));
+
+		// Edit.
+		FUIAction EditAction(FExecuteAction::CreateSP(this, &SMediaSourceManagerPreview::OnEditInput));
+		MenuBuilder.AddMenuEntry(LOCTEXT("Edit", "Edit"),
+			LOCTEXT("EditToolTip", "Edit this asset"),
+			FSlateIcon(FAppStyle::GetAppStyleSetName(), "Icons.Edit"), EditAction);
+
+		// Clear.
+		FUIAction ClearAction(FExecuteAction::CreateSP(this, &SMediaSourceManagerPreview::ClearInput));
+		MenuBuilder.AddMenuEntry(LOCTEXT("Clear", "Clear"),
+			LOCTEXT("ClearToolTip", "Clears the asset set on this field"),
+			FSlateIcon(FAppStyle::GetAppStyleSetName(), "GenericCommands.Delete"), ClearAction);
+
+		MenuBuilder.EndSection();
+	}
+	
+	// Bring up menu.
+	FWidgetPath WidgetPath = MouseEvent.GetEventPath() != nullptr ? *MouseEvent.GetEventPath() : FWidgetPath();
+	FSlateApplication::Get().PushMenu(SharedThis(this), WidgetPath, MenuBuilder.MakeWidget(), FSlateApplication::Get().GetCursorPos(), FPopupTransitionEffect::ContextMenu);
+}
+
+
 
 #undef LOCTEXT_NAMESPACE
diff --git a/Engine/Plugins/Media/MediaSourceManager/Source/MediaSourceManagerEditor/Private/Widgets/SMediaSourceManagerPreview.h b/Engine/Plugins/Media/MediaSourceManager/Source/MediaSourceManagerEditor/Private/Widgets/SMediaSourceManagerPreview.h
--- a/Engine/Plugins/Media/MediaSourceManager/Source/MediaSourceManagerEditor/Private/Widgets/SMediaSourceManagerPreview.h
+++ b/Engine/Plugins/Media/MediaSourceManager/Source/MediaSourceManagerEditor/Private/Widgets/SMediaSourceManagerPreview.h
@@ -38,11 +38,17 @@ public:
 	virtual FReply OnDrop(const FGeometry& MyGeometry, const FDragDropEvent& DragDropEvent) override;
 	
 	virtual FReply OnMouseButtonDown(const FGeometry& MyGeometry, const FPointerEvent& MouseEvent) override;
+	virtual FReply OnMouseButtonUp(const FGeometry& MyGeometry, const FPointerEvent& MouseEvent) override;
 	virtual FReply OnDragDetected(const FGeometry& MyGeometry, const FPointerEvent& MouseEvent) override;
 	//~ End of SWidget interface
 
 private:
 
+	/**
+	 * Clears the input on this channel.
+	 */
+	void ClearInput();
+
 	/**
 	 * Assigns a media source to this channel.
 	 *
@@ -50,11 +56,21 @@ private:
 	 */
 	void AssignMediaSourceInput(UMediaSource* MediaSource);
 
+	/**
+	 * Call this to bring up an editor to edit the input.
+	 */
+	void OnEditInput();
+
 	/**
 	 * Gets a material for this texture.
 	 */
 	UMaterialInstanceConstant* GetMaterial();
 
+	/**
+	 * Call this to open a context menu.
+	 */
+	void OpenContextMenu(const FPointerEvent& MouseEvent);
+
 	/** Pointer to the object that is being viewed. */
 	TWeakObjectPtr<UMediaSourceManagerChannel> ChannelPtr;
 
