MediaSourceManager: Copied over assign input functionality to the preview widget.
#rnx

diff --git a/Engine/Plugins/Media/MediaSourceManager/Source/MediaSourceManagerEditor/Private/Widgets/SMediaSourceManagerPreview.cpp b/Engine/Plugins/Media/MediaSourceManager/Source/MediaSourceManagerEditor/Private/Widgets/SMediaSourceManagerPreview.cpp
--- a/Engine/Plugins/Media/MediaSourceManager/Source/MediaSourceManagerEditor/Private/Widgets/SMediaSourceManagerPreview.cpp
+++ b/Engine/Plugins/Media/MediaSourceManager/Source/MediaSourceManagerEditor/Private/Widgets/SMediaSourceManagerPreview.cpp
@@ -2,17 +2,24 @@
 
 #include "Widgets/SMediaSourceManagerPreview.h"
 
+#include "ContentBrowserModule.h"
 #include "DragAndDrop/AssetDragDropOp.h"
 #include "Editor.h"
 #include "Framework/Application/SlateApplication.h"
 #include "Framework/MultiBox/MultiBoxBuilder.h"
+#include "Framework/Notifications/NotificationManager.h"
+#include "IContentBrowserSingleton.h"
+#include "IMediaIOCoreDeviceProvider.h"
+#include "IMediaIOCoreModule.h"
 #include "Inputs/MediaSourceManagerInputMediaSource.h"
 #include "Materials/MaterialInstanceConstant.h"
+#include "MediaIOPermutationsSelectorBuilder.h"
 #include "MediaPlayer.h"
 #include "MediaSourceManagerChannel.h"
 #include "MediaSource.h"
 #include "MediaTexture.h"
 #include "Subsystems/AssetEditorSubsystem.h"
+#include "Widgets/Notifications/SNotificationList.h"
 #include "Widgets/SMediaPlayerEditorViewer.h"
 
 #define LOCTEXT_NAMESPACE "SMediaSourceManagerPreview"
@@ -158,6 +165,97 @@ void SMediaSourceManagerPreview::AssignMediaSourceInput(UMediaSource* MediaSourc
 	}
 }
 
+void SMediaSourceManagerPreview::AssignMediaIOInput(IMediaIOCoreDeviceProvider* DeviceProvider,
+	FMediaIOConnection Connection)
+{
+	UMediaSourceManagerChannel* Channel = ChannelPtr.Get();
+	if (Channel != nullptr)
+	{
+		// Create media source.
+		FMediaIOConfiguration Configuration;
+		Configuration = DeviceProvider->GetDefaultConfiguration();
+		Configuration.MediaConnection = Connection;
+		UMediaSource* MediaSource = DeviceProvider->CreateMediaSource(Configuration, Channel);
+		if (MediaSource != nullptr)
+		{
+			AssignMediaSourceInput(MediaSource);
+		}
+		else
+		{
+			// Failed to create media source.
+			// Remove any existing error.
+			DismissErrorNotification();
+
+			// Inform the user.
+			FNotificationInfo Info(LOCTEXT("FailedToCreateMediaSource", "Failed to create a Media Source."));
+			Info.bFireAndForget = false;
+			Info.bUseLargeFont = false;
+			Info.bUseThrobber = false;
+			Info.FadeOutDuration = 0.25f;
+			Info.ButtonDetails.Add(FNotificationButtonInfo(LOCTEXT("Dismiss", "Dismiss"), LOCTEXT("DismissToolTip", "Dismiss this notification."),
+				FSimpleDelegate::CreateSP(this, &SMediaSourceManagerPreview::DismissErrorNotification)));
+
+			ErrorNotificationPtr = FSlateNotificationManager::Get().AddNotification(Info);
+			TSharedPtr<SNotificationItem> ErrorNotification = ErrorNotificationPtr.Pin();
+			if (ErrorNotification != nullptr)
+			{
+				ErrorNotification->SetCompletionState(SNotificationItem::CS_Pending);
+			}
+		}
+	}
+}
+
+TSharedRef<SWidget> SMediaSourceManagerPreview::BuildMediaSourcePickerWidget()
+{
+	FAssetPickerConfig AssetPickerConfig;
+	{
+		AssetPickerConfig.OnAssetSelected = FOnAssetSelected::CreateRaw(this,
+			&SMediaSourceManagerPreview::AddMediaSource);
+		AssetPickerConfig.OnAssetEnterPressed = FOnAssetEnterPressed::CreateRaw(this,
+			&SMediaSourceManagerPreview::AddMediaSourceEnterPressed);
+		AssetPickerConfig.bAllowNullSelection = false;
+		AssetPickerConfig.InitialAssetViewType = EAssetViewType::List;
+		AssetPickerConfig.Filter.bRecursiveClasses = true;
+		AssetPickerConfig.Filter.ClassPaths.Add(UMediaSource::StaticClass()->GetClassPathName());
+		AssetPickerConfig.SaveSettingsName = TEXT("MediaSourceManagerAssetPicker");
+	}
+
+	FContentBrowserModule& ContentBrowserModule = FModuleManager::Get().LoadModuleChecked<FContentBrowserModule>(TEXT("ContentBrowser"));
+
+	TSharedRef<SBox> Picker = SNew(SBox)
+		.WidthOverride(300.0f)
+		.HeightOverride(300.f)
+		[
+			ContentBrowserModule.Get().CreateAssetPicker(AssetPickerConfig)
+		];
+
+	return Picker;
+}
+
+void SMediaSourceManagerPreview::AddMediaSource(const FAssetData& AssetData)
+{
+	FSlateApplication::Get().DismissAllMenus();
+
+	// Get media source from the asset..
+	UObject* SelectedObject = AssetData.GetAsset();
+	if (SelectedObject)
+	{
+		UMediaSource* MediaSource = Cast<UMediaSource>(AssetData.GetAsset());
+		if (MediaSource != nullptr)
+		{
+			AssignMediaSourceInput(MediaSource);
+		}
+	}
+}
+
+void SMediaSourceManagerPreview::AddMediaSourceEnterPressed(const TArray<FAssetData>& AssetData)
+{
+	if (AssetData.Num() > 0)
+	{
+		AddMediaSource(AssetData[0].GetAsset());
+	}
+}
+
 void SMediaSourceManagerPreview::OnEditInput()
 {
 	// Get our input.
@@ -251,12 +349,77 @@ void SMediaSourceManagerPreview::OpenContextMenu(const FPointerEvent& MouseEvent
 
 		MenuBuilder.EndSection();
 	}
+
+	// Get all Media IO device providers.
+	IMediaIOCoreModule& MediaIOCoreModule = IMediaIOCoreModule::Get();
+	TConstArrayView<IMediaIOCoreDeviceProvider*> DeviceProviders =
+		MediaIOCoreModule.GetDeviceProviders();
+
+	// Loop through each provider.
+	for (IMediaIOCoreDeviceProvider* DeviceProvider : DeviceProviders)
+	{
+		if (DeviceProvider != nullptr)
+		{
+			// Start menu section.
+			FName ProviderName = DeviceProvider->GetFName();
+			MenuBuilder.BeginSection(ProviderName, FText::FromName(ProviderName));
+
+			// Go over all connections.
+			TArray<FMediaIOConnection> Connections = DeviceProvider->GetConnections();
+			for (const FMediaIOConnection& Connection : Connections)
+			{
+				// Add this connection.
+				FText DeviceName = FText::FromName(Connection.Device.DeviceName);
+				FText LinkName = FMediaIOPermutationsSelectorBuilder::GetLabel(
+					FMediaIOPermutationsSelectorBuilder::NAME_TransportType,
+					Connection);
+				FText MenuText;
+				if (DeviceProvider->ShowInputTransportInSelector())
+				{
+					MenuText = FText::Format(LOCTEXT("Connection", "{0}: {1}"),
+						DeviceName, LinkName);
+				}
+				else
+				{
+					MenuText = DeviceName;
+				}
+
+				FUIAction AssignMediaIOInputAction(FExecuteAction::CreateSP(this,
+					&SMediaSourceManagerPreview::AssignMediaIOInput, DeviceProvider, Connection));
+				MenuBuilder.AddMenuEntry(MenuText, FText(), FSlateIcon(),
+					AssignMediaIOInputAction);
+			}
+
+			MenuBuilder.EndSection();
+		}
+	}
+
+	// Add assets.
+	MenuBuilder.BeginSection(NAME_None, LOCTEXT("MediaSourceAssets", "Media Source Assets"));
+	auto SubMenuCallback = [this](FMenuBuilder& SubMenuBuilder)
+	{
+		SubMenuBuilder.AddWidget(BuildMediaSourcePickerWidget(), FText::GetEmpty(), true);
+	};
+	MenuBuilder.AddSubMenu(
+		LOCTEXT("SelectAsset", "Select Asset"),
+		LOCTEXT("SelectAsset_ToolTip", "Select an existing Media Source asset."),
+		FNewMenuDelegate::CreateLambda(SubMenuCallback)
+	);
+	MenuBuilder.EndSection();
 	
 	// Bring up menu.
 	FWidgetPath WidgetPath = MouseEvent.GetEventPath() != nullptr ? *MouseEvent.GetEventPath() : FWidgetPath();
 	FSlateApplication::Get().PushMenu(SharedThis(this), WidgetPath, MenuBuilder.MakeWidget(), FSlateApplication::Get().GetCursorPos(), FPopupTransitionEffect::ContextMenu);
 }
 
-
+void SMediaSourceManagerPreview::DismissErrorNotification()
+{
+	TSharedPtr<SNotificationItem> ErrorNotification = ErrorNotificationPtr.Pin();
+	if (ErrorNotification != nullptr)
+	{
+		ErrorNotification->ExpireAndFadeout();
+		ErrorNotification.Reset();
+	}
+}
 
 #undef LOCTEXT_NAMESPACE
diff --git a/Engine/Plugins/Media/MediaSourceManager/Source/MediaSourceManagerEditor/Private/Widgets/SMediaSourceManagerPreview.h b/Engine/Plugins/Media/MediaSourceManager/Source/MediaSourceManagerEditor/Private/Widgets/SMediaSourceManagerPreview.h
--- a/Engine/Plugins/Media/MediaSourceManager/Source/MediaSourceManagerEditor/Private/Widgets/SMediaSourceManagerPreview.h
+++ b/Engine/Plugins/Media/MediaSourceManager/Source/MediaSourceManagerEditor/Private/Widgets/SMediaSourceManagerPreview.h
@@ -5,11 +5,16 @@
 #include "Widgets/DeclarativeSyntaxSupport.h"
 #include "Widgets/SCompoundWidget.h"
 
+class IMediaIOCoreDeviceProvider;
 class ISlateStyle;
+class SNotificationItem;
 class UMaterialInstanceConstant;
 class UMediaSource;
 class UMediaSourceManagerChannel;
 
+struct FAssetData;
+struct FMediaIOConnection;
+
 /**
  * Implements a preview for a single media source manager channel.
  */
@@ -56,6 +61,29 @@ private:
 	 */
 	void AssignMediaSourceInput(UMediaSource* MediaSource);
 
+	/**
+	 * Assigns a MediaIO input to this channel.
+	 *
+	 * @param DeviceProvider	MediaIO device provider to get the input from.
+	 * @param Connection		Connection to pass to the device provider.
+	 */
+	void AssignMediaIOInput(IMediaIOCoreDeviceProvider* DeviceProvider, FMediaIOConnection Connection);
+
+	/**
+	 * Creates a widget to select a media source.
+	 */
+	TSharedRef<SWidget> BuildMediaSourcePickerWidget();
+
+	/**
+	 * Callback to add a media source.
+	 */
+	void AddMediaSource(const FAssetData& AssetData);
+
+	/**
+	 * Callback to add a media source.
+	 */
+	void AddMediaSourceEnterPressed(const TArray<FAssetData>& AssetData);
+
 	/**
 	 * Call this to bring up an editor to edit the input.
 	 */
@@ -71,8 +99,15 @@ private:
 	 */
 	void OpenContextMenu(const FPointerEvent& MouseEvent);
 
+	/**
+	 * Call this to remove the error notification.
+	 */
+	void DismissErrorNotification();
+
 	/** Pointer to the object that is being viewed. */
 	TWeakObjectPtr<UMediaSourceManagerChannel> ChannelPtr;
+	/** Holds our error notification so we can dismiss it. */
+	TWeakPtr<SNotificationItem> ErrorNotificationPtr;
 
 	/** Name for the media texture parameter in the material. */
 	static FLazyName MediaTextureName;
