MediaSourceManager: Tells media plate not to modify its material so we can use our material that points to our texture.

#rnx

diff --git a/Engine/Plugins/Media/MediaPlate/Source/MediaPlate/Private/MediaPlate.cpp b/Engine/Plugins/Media/MediaPlate/Source/MediaPlate/Private/MediaPlate.cpp
--- a/Engine/Plugins/Media/MediaPlate/Source/MediaPlate/Private/MediaPlate.cpp
+++ b/Engine/Plugins/Media/MediaPlate/Source/MediaPlate/Private/MediaPlate.cpp
@@ -130,7 +130,19 @@ void AMediaPlate::ApplyMaterial(UMaterialInterface* Material)
 		UMaterialInstanceDynamic* MID = Cast<UMaterialInstanceDynamic>(Material);
 		UMaterialInterface* Result = nullptr;
 
-		if (MID != nullptr)
+		// See if we can modify this material.
+		bool bCanModify = true;
+		FMediaPlateModule* MediaPlateModule = FModuleManager::GetModulePtr<FMediaPlateModule>("MediaPlate");
+		if (MediaPlateModule != nullptr)
+		{
+			MediaPlateModule->OnMediaPlateApplyMaterial.Broadcast(Material, this, bCanModify);
+		}
+		
+		if (bCanModify == false)
+		{
+			LastMaterial = Material;
+		}
+		else if (MID != nullptr)
 		{
 			MID->SetTextureParameterValue(MediaTextureName, MediaPlateComponent->GetMediaTexture());
 
diff --git a/Engine/Plugins/Media/MediaPlate/Source/MediaPlate/Public/MediaPlateModule.h b/Engine/Plugins/Media/MediaPlate/Source/MediaPlate/Public/MediaPlateModule.h
--- a/Engine/Plugins/Media/MediaPlate/Source/MediaPlate/Public/MediaPlateModule.h
+++ b/Engine/Plugins/Media/MediaPlate/Source/MediaPlate/Public/MediaPlateModule.h
@@ -2,15 +2,25 @@
 
 #pragma once
 
+#include "Delegates/Delegate.h"
 #include "Logging/LogMacros.h"
 #include "Modules/ModuleManager.h"
 
+class AMediaPlate;
+class UMaterialInterface;
 class UMediaPlayer;
 class UObject;
 
 /** Log category for this module. */
 DECLARE_LOG_CATEGORY_EXTERN(LogMediaPlate, Log, All);
 
+/**
+ * Callback when a media plate applies a material to itself.
+ * Set bCanModify to false if you do not want the media plate to modify the material so it can show
+ * the media plate's texture.
+ */
+DECLARE_MULTICAST_DELEGATE_ThreeParams(FOnMediaPlateApplyMaterial, UMaterialInterface*, AMediaPlate*, bool& /*bCanNodify*/)
+
 class MEDIAPLATE_API FMediaPlateModule : public IModuleInterface
 {
 public:
@@ -23,6 +33,9 @@ public:
 	virtual void StartupModule() override;
 	virtual void ShutdownModule() override;
 
+	/** Add to this if you want a callback when a media plate applies a material to itself. */
+	FOnMediaPlateApplyMaterial OnMediaPlateApplyMaterial;
+
 private:
 	/** ID for our delegate. */
 	int32 GetPlayerFromObjectID = INDEX_NONE;
diff --git a/Engine/Plugins/Media/MediaSourceManager/Source/MediaSourceManager/MediaSourceManager.Build.cs b/Engine/Plugins/Media/MediaSourceManager/Source/MediaSourceManager/MediaSourceManager.Build.cs
--- a/Engine/Plugins/Media/MediaSourceManager/Source/MediaSourceManager/MediaSourceManager.Build.cs
+++ b/Engine/Plugins/Media/MediaSourceManager/Source/MediaSourceManager/MediaSourceManager.Build.cs
@@ -16,6 +16,7 @@ public class MediaSourceManager : ModuleRules
 				"Engine",
 				"MediaAssets",
 				"MediaFrameworkUtilities",
+				"MediaPlate",
 			}
 			);
 	}
diff --git a/Engine/Plugins/Media/MediaSourceManager/Source/MediaSourceManager/Private/MediaSourceManager.cpp b/Engine/Plugins/Media/MediaSourceManager/Source/MediaSourceManager/Private/MediaSourceManager.cpp
--- a/Engine/Plugins/Media/MediaSourceManager/Source/MediaSourceManager/Private/MediaSourceManager.cpp
+++ b/Engine/Plugins/Media/MediaSourceManager/Source/MediaSourceManager/Private/MediaSourceManager.cpp
@@ -2,8 +2,37 @@
 
 #include "MediaSourceManager.h"
 
+#include "MediaPlateModule.h"
+#include "Materials/MaterialInstance.h"
+#include "Materials/MaterialInstanceConstant.h"
+
 #define LOCTEXT_NAMESPACE "MediaSourceManager"
 
+UMediaSourceManager::UMediaSourceManager(const FObjectInitializer& ObjectInitializer)
+	: Super(ObjectInitializer)
+{
+	if (!HasAnyFlags(RF_ClassDefaultObject))
+	{
+		FMediaPlateModule* MediaPlateModule = FModuleManager::LoadModulePtr<FMediaPlateModule>("MediaPlate");
+		if (MediaPlateModule != nullptr)
+		{
+			MediaPlateModule->OnMediaPlateApplyMaterial.AddUObject(this,
+				&UMediaSourceManager::OnMediaPlateApplyMaterial);
+		}
+	}
+}
+
+void UMediaSourceManager::BeginDestroy()
+{
+	FMediaPlateModule* MediaPlateModule = FModuleManager::GetModulePtr<FMediaPlateModule>("MediaPlate");
+	if (MediaPlateModule != nullptr)
+	{
+		MediaPlateModule->OnMediaPlateApplyMaterial.RemoveAll(this);
+	}
+
+	Super::BeginDestroy();
+}
+
 void UMediaSourceManager::Validate()
 {
 	// Check we have some channels.
@@ -26,5 +55,19 @@ void UMediaSourceManager::Validate()
 	}
 }
 
+void UMediaSourceManager::OnMediaPlateApplyMaterial(UMaterialInterface* Material,
+	AMediaPlate* MediaPlate, bool& bCanModify)
+{
+	// Is this material one of ours?
+	for (UMediaSourceManagerChannel* Channel : Channels)
+	{
+		if (Channel->Material.Get() == Material)
+		{
+			bCanModify = false;
+			break;
+		}
+	}
+}
+
 
 #undef LOCTEXT_NAMESPACE
diff --git a/Engine/Plugins/Media/MediaSourceManager/Source/MediaSourceManager/Public/MediaSourceManager.h b/Engine/Plugins/Media/MediaSourceManager/Source/MediaSourceManager/Public/MediaSourceManager.h
--- a/Engine/Plugins/Media/MediaSourceManager/Source/MediaSourceManager/Public/MediaSourceManager.h
+++ b/Engine/Plugins/Media/MediaSourceManager/Source/MediaSourceManager/Public/MediaSourceManager.h
@@ -9,6 +9,9 @@
 
 #include "MediaSourceManager.generated.h"
 
+class AMediaPlate;
+class UMaterialInterface;
+
 /**
 * Manager to handle media sources and their connections.
 */
@@ -18,6 +21,13 @@ class MEDIASOURCEMANAGER_API UMediaSourceManager : public UObject
 	GENERATED_BODY()
 
 public:
+
+	UMediaSourceManager(const FObjectInitializer& ObjectInitializer);
+
+	//~ Begin UObject interface.
+	virtual void BeginDestroy() override;
+	//~ End UObject interface.
+
 	/** Our channels. */
 	UPROPERTY(EditAnywhere, BlueprintReadWrite, Category = "Channels")
 	TArray<TObjectPtr<UMediaSourceManagerChannel>> Channels;
@@ -27,4 +37,11 @@ public:
 	 */
 	void Validate();
 
+private:
+
+	/**
+	 * Call this to see if a material is ours and should not be modified.
+	 */
+	void OnMediaPlateApplyMaterial(UMaterialInterface* Material, AMediaPlate* MediaPlate, bool& bCanModify);
+
 };
