MediaSourceManager: Added mediaplate to the plugin list.
#rnx

diff --git a/Engine/Plugins/Media/MediaSourceManager/MediaSourceManager.uplugin b/Engine/Plugins/Media/MediaSourceManager/MediaSourceManager.uplugin
--- a/Engine/Plugins/Media/MediaSourceManager/MediaSourceManager.uplugin
+++ b/Engine/Plugins/Media/MediaSourceManager/MediaSourceManager.uplugin
@@ -31,6 +31,10 @@
 		{
 			"Name": "MediaFrameworkUtilities",
 			"Enabled": true
+		},
+		{
+			"Name": "MediaPlate",
+			"Enabled": true
 		}
 	]
 }
\ No newline at end of file
